Shots: Add new shots plugin to VP plugins.
Add new DirectoryPlaceholder asset type, used as an extremely lightweight placeholder to place in otherwise empty directories, allowing them to be submitted to source control.

diff --git a/Engine/Plugins/VirtualProduction/Shots/Shots.uplugin b/Engine/Plugins/VirtualProduction/Shots/Shots.uplugin
new file mode 100644
--- /dev/null
+++ b/Engine/Plugins/VirtualProduction/Shots/Shots.uplugin
@@ -0,0 +1,29 @@
+{
+	"FileVersion": 3,
+	"Version": 1,
+	"VersionName": "1.0",
+	"FriendlyName": "Shots Management",
+	"Description": "Tools for creating, editing, and managing shots within Unreal Editor",
+	"Category": "Virtual Production",
+	"CreatedBy": "Epic Games, Inc.",
+	"CreatedByURL": "https://epicgames.com",
+	"DocsURL": "",
+	"MarketplaceURL": "",
+	"SupportURL": "",
+	"CanContainContent": true,
+	"IsBetaVersion": false,
+	"IsExperimentalVersion": true,
+	"Hidden": false,
+	"Installed": false,
+	"SupportedPrograms": [ "LiveLinkHub" ],
+	"Modules": [
+		{
+			"Name": "ShotsEditor",
+			"Type": "Editor",
+			"ProgramAllowlist": [ "LiveLinkHub" ],
+			"LoadingPhase": "Default"
+		}
+	],
+	"Plugins": [
+	]
+}
\ No newline at end of file
diff --git a/Engine/Plugins/VirtualProduction/Shots/Source/ShotsEditor/Private/AssetDefinition_DirectoryPlaceholder.cpp b/Engine/Plugins/VirtualProduction/Shots/Source/ShotsEditor/Private/AssetDefinition_DirectoryPlaceholder.cpp
new file mode 100644
--- /dev/null
+++ b/Engine/Plugins/VirtualProduction/Shots/Source/ShotsEditor/Private/AssetDefinition_DirectoryPlaceholder.cpp
@@ -0,0 +1,29 @@
+// Copyright Epic Games, Inc. All Rights Reserved.
+
+#include "AssetDefinition_DirectoryPlaceholder.h"
+#include "DirectoryPlaceholder.h"
+
+#define LOCTEXT_NAMESPACE "AssetTypeActions"
+
+FText UAssetDefinition_DirectoryPlaceholder::GetAssetDisplayName() const
+{
+	return LOCTEXT("AssetTypeActions_DirectoryPlaceholder", "Directory Placeholder");
+}
+
+FLinearColor UAssetDefinition_DirectoryPlaceholder::GetAssetColor() const
+{
+	return FLinearColor::Gray;
+}
+
+TSoftClassPtr<> UAssetDefinition_DirectoryPlaceholder::GetAssetClass() const
+{
+	return UDirectoryPlaceholder::StaticClass();
+}
+
+TConstArrayView<FAssetCategoryPath> UAssetDefinition_DirectoryPlaceholder::GetAssetCategories() const
+{
+	static const FAssetCategoryPath Categories[] = { EAssetCategoryPaths::Misc };
+	return Categories;
+}
+
+#undef LOCTEXT_NAMESPACE
diff --git a/Engine/Plugins/VirtualProduction/Shots/Source/ShotsEditor/Private/AssetDefinition_DirectoryPlaceholder.h b/Engine/Plugins/VirtualProduction/Shots/Source/ShotsEditor/Private/AssetDefinition_DirectoryPlaceholder.h
new file mode 100644
--- /dev/null
+++ b/Engine/Plugins/VirtualProduction/Shots/Source/ShotsEditor/Private/AssetDefinition_DirectoryPlaceholder.h
@@ -0,0 +1,22 @@
+// Copyright Epic Games, Inc. All Rights Reserved.
+
+#pragma once
+
+#include "AssetDefinitionDefault.h"
+
+#include "AssetDefinition_DirectoryPlaceholder.generated.h"
+
+/** Asset definition for a UDirectoryPlaceholder asset */
+UCLASS()
+class UAssetDefinition_DirectoryPlaceholder : public UAssetDefinitionDefault
+{
+	GENERATED_BODY()
+
+protected:
+	// Begin UAssetDefinition Interface
+	virtual FText GetAssetDisplayName() const override;
+	virtual FLinearColor GetAssetColor() const override;
+	virtual TSoftClassPtr<UObject> GetAssetClass() const override;
+	virtual TConstArrayView<FAssetCategoryPath> GetAssetCategories() const override;
+	// End UAssetDefinition Interface
+};
diff --git a/Engine/Plugins/VirtualProduction/Shots/Source/ShotsEditor/Private/DirectoryPlaceholder.cpp b/Engine/Plugins/VirtualProduction/Shots/Source/ShotsEditor/Private/DirectoryPlaceholder.cpp
new file mode 100644
--- /dev/null
+++ b/Engine/Plugins/VirtualProduction/Shots/Source/ShotsEditor/Private/DirectoryPlaceholder.cpp
@@ -0,0 +1,41 @@
+// Copyright Epic Games, Inc. All Rights Reserved.
+
+#include "DirectoryPlaceholder.h"
+#include "FrontendFilterBase.h"
+
+#define LOCTEXT_NAMESPACE "DirectoryPlaceholder"
+
+class FFrontendFilter_DirectoryPlaceholder : public FFrontendFilter
+{
+public:
+	FFrontendFilter_DirectoryPlaceholder(TSharedPtr<FFrontendFilterCategory> InCategory)
+		: FFrontendFilter(InCategory)
+	{
+	}
+
+	// Begin FFilterBase interface
+	virtual FString GetName() const override { return TEXT("DirectoryPlaceholderFilter"); }
+	virtual FText GetDisplayName() const override { return LOCTEXT("DirectoryPlaceholderFilterName", "Show Directory Placeholders"); }
+	virtual FText GetToolTipText() const override { return LOCTEXT("DirectoryPlaceholderFilterTooltip", "Show Directory Placeholders"); }
+
+	/** This is an inverse filter to prevent the asset view from recursively displaying all assets */
+	virtual bool IsInverseFilter() const override { return true; } 
+	// End FFilterBase interface
+
+	// Begin IFilter interface
+	virtual bool PassesFilter(FAssetFilterType InItem) const override;
+	// End IFilter interface
+};
+
+bool FFrontendFilter_DirectoryPlaceholder::PassesFilter(FAssetFilterType InItem) const
+{
+	const FContentBrowserItemDataAttributeValue ClassValue = InItem.GetItemAttribute(NAME_Class);
+	return !ClassValue.IsValid() || ClassValue.GetValue<FString>() != UDirectoryPlaceholder::StaticClass()->GetPathName();
+}
+
+void UDirectoryPlaceholderSearchFilter::AddFrontEndFilterExtensions(TSharedPtr<FFrontendFilterCategory> DefaultCategory, TArray<TSharedRef<class FFrontendFilter>>& InOutFilterList) const
+{
+	InOutFilterList.Add(MakeShareable(new FFrontendFilter_DirectoryPlaceholder(DefaultCategory)));
+}
+
+#undef LOCTEXT_NAMESPACE
diff --git a/Engine/Plugins/VirtualProduction/Shots/Source/ShotsEditor/Private/DirectoryPlaceholder.h b/Engine/Plugins/VirtualProduction/Shots/Source/ShotsEditor/Private/DirectoryPlaceholder.h
new file mode 100644
--- /dev/null
+++ b/Engine/Plugins/VirtualProduction/Shots/Source/ShotsEditor/Private/DirectoryPlaceholder.h
@@ -0,0 +1,35 @@
+// Copyright Epic Games, Inc. All Rights Reserved.
+
+#pragma once
+
+#include "ContentBrowserFrontEndFilterExtension.h"
+
+#include "DirectoryPlaceholder.generated.h"
+
+class FFrontendFilter;
+class FFrontendFilterCategory;
+
+/** 
+ * Extremely lightweight object that can be used as a placeholder in an otherwise empty directory.
+ * The presence of a placeholder object allows that directory to be submitted to source control.
+ */
+UCLASS()
+class UDirectoryPlaceholder : public UObject
+{
+	GENERATED_BODY()
+
+public:
+	UDirectoryPlaceholder() = default;
+};
+
+/** Content Browser filter used to show/hide directory placeholder assets */
+UCLASS()
+class UDirectoryPlaceholderSearchFilter : public UContentBrowserFrontEndFilterExtension
+{
+	GENERATED_BODY()
+
+protected:
+	// Begin UContentBrowserFrontEndFilterExtension Interface
+	virtual void AddFrontEndFilterExtensions(TSharedPtr<FFrontendFilterCategory> DefaultCategory, TArray<TSharedRef<FFrontendFilter>>& InOutFilterList) const override;
+	// End UContentBrowserFrontEndFilterExtension Interface
+};
diff --git a/Engine/Plugins/VirtualProduction/Shots/Source/ShotsEditor/Private/DirectoryPlaceholderFactory.cpp b/Engine/Plugins/VirtualProduction/Shots/Source/ShotsEditor/Private/DirectoryPlaceholderFactory.cpp
new file mode 100644
--- /dev/null
+++ b/Engine/Plugins/VirtualProduction/Shots/Source/ShotsEditor/Private/DirectoryPlaceholderFactory.cpp
@@ -0,0 +1,26 @@
+// Copyright Epic Games, Inc. All Rights Reserved.
+
+#include "DirectoryPlaceholderFactory.h"
+#include "DirectoryPlaceholder.h"
+
+#define LOCTEXT_NAMESPACE "DirectoryPlaceholderFactory"
+
+UDirectoryPlaceholderFactory::UDirectoryPlaceholderFactory()
+{
+	SupportedClass = UDirectoryPlaceholder::StaticClass();
+	bCreateNew = true;
+	bEditAfterNew = true;
+}
+
+bool UDirectoryPlaceholderFactory::CanCreateNew() const
+{
+	return true;
+}
+
+UObject* UDirectoryPlaceholderFactory::FactoryCreateNew(UClass* Class, UObject* InParent, FName Name, EObjectFlags Flags, UObject* Context, FFeedbackContext* Warn)
+{
+	check(Class->IsChildOf(UDirectoryPlaceholder::StaticClass()));
+	return NewObject<UDirectoryPlaceholder>(InParent, Class, Name, Flags);
+}
+
+#undef LOCTEXT_NAMESPACE
diff --git a/Engine/Plugins/VirtualProduction/Shots/Source/ShotsEditor/Private/DirectoryPlaceholderFactory.h b/Engine/Plugins/VirtualProduction/Shots/Source/ShotsEditor/Private/DirectoryPlaceholderFactory.h
new file mode 100644
--- /dev/null
+++ b/Engine/Plugins/VirtualProduction/Shots/Source/ShotsEditor/Private/DirectoryPlaceholderFactory.h
@@ -0,0 +1,22 @@
+// Copyright Epic Games, Inc. All Rights Reserved.
+
+#pragma once
+
+#include "Factories/Factory.h"
+
+#include "DirectoryPlaceholderFactory.generated.h"
+
+/** Implements a factory for UDirectoryPlaceholder objects */
+UCLASS(hidecategories=Object)
+class UDirectoryPlaceholderFactory : public UFactory
+{
+	GENERATED_BODY()
+
+protected:
+	UDirectoryPlaceholderFactory();
+
+	// Begin UFactory Interface
+	virtual UObject* FactoryCreateNew(UClass* Class, UObject* InParent, FName Name, EObjectFlags Flags, UObject* Context, FFeedbackContext* Warn) override;
+	virtual bool CanCreateNew() const override;
+	// End UFactory Interface
+};
diff --git a/Engine/Plugins/VirtualProduction/Shots/Source/ShotsEditor/Private/ShotsEditorModule.cpp b/Engine/Plugins/VirtualProduction/Shots/Source/ShotsEditor/Private/ShotsEditorModule.cpp
new file mode 100644
--- /dev/null
+++ b/Engine/Plugins/VirtualProduction/Shots/Source/ShotsEditor/Private/ShotsEditorModule.cpp
@@ -0,0 +1,5 @@
+// Copyright Epic Games, Inc. All Rights Reserved.
+
+#include "Modules/ModuleManager.h"
+
+IMPLEMENT_MODULE(FDefaultModuleImpl, ShotsEditor);
diff --git a/Engine/Plugins/VirtualProduction/Shots/Source/ShotsEditor/ShotsEditor.build.cs b/Engine/Plugins/VirtualProduction/Shots/Source/ShotsEditor/ShotsEditor.build.cs
new file mode 100644
--- /dev/null
+++ b/Engine/Plugins/VirtualProduction/Shots/Source/ShotsEditor/ShotsEditor.build.cs
@@ -0,0 +1,28 @@
+// Copyright Epic Games, Inc. All Rights Reserved.
+
+namespace UnrealBuildTool.Rules
+{
+	public class ShotsEditor : ModuleRules
+	{
+		public ShotsEditor(ReadOnlyTargetRules Target) : base(Target)
+		{
+			PublicDependencyModuleNames.AddRange(
+				new string[]
+				{
+				}
+			);
+
+			PrivateDependencyModuleNames.AddRange(
+				new string[]
+				{
+					"AssetDefinition",
+					"ContentBrowser",
+					"ContentBrowserData",
+					"Core",
+					"CoreUObject",
+					"UnrealEd",
+				}
+			);
+		}
+	}
+}
