ShotManagement: Add Production Settings to the ShotsEditor module, with an array of project-wide productions, each with sequencer setting overrides. These overrides include Default Display Rate and Start Frame for new level sequences, and Default Hierarchical Bias for new subsequences.

diff --git a/Engine/Plugins/VirtualProduction/Shots/Shots.uplugin b/Engine/Plugins/VirtualProduction/Shots/Shots.uplugin
--- a/Engine/Plugins/VirtualProduction/Shots/Shots.uplugin
+++ b/Engine/Plugins/VirtualProduction/Shots/Shots.uplugin
@@ -25,5 +25,9 @@
 		}
 	],
 	"Plugins": [
+		{
+			"Name": "LevelSequenceEditor",
+			"Enabled": true
+		}
 	]
 }
\ No newline at end of file
diff --git a/Engine/Plugins/VirtualProduction/Shots/Source/ShotsEditor/Private/ProductionSettings.cpp b/Engine/Plugins/VirtualProduction/Shots/Source/ShotsEditor/Private/ProductionSettings.cpp
new file mode 100644
--- /dev/null
+++ b/Engine/Plugins/VirtualProduction/Shots/Source/ShotsEditor/Private/ProductionSettings.cpp
@@ -0,0 +1,268 @@
+// Copyright Epic Games, Inc. All Rights Reserved.
+#include "ProductionSettings.h"
+
+#include "HAL/FileManager.h"
+#include "HAL/PlatformFileManager.h"
+#include "LevelSequenceProjectSettings.h"
+#include "Misc/ConfigCacheIni.h"
+#include "MovieSceneToolsProjectSettings.h"
+#include "Sections/MovieSceneSubSection.h"
+#include "SourceControlHelpers.h"
+
+#define LOCTEXT_NAMESPACE "ProductionSettings"
+
+DEFINE_LOG_CATEGORY_STATIC(LogProductionSettings, Log, All)
+
+namespace UE::Shots::ProductionSettings::Private
+{
+	constexpr int32 TopDownHBiasValue = -100;
+	constexpr int32 BottomUpHBiasValue = 100;
+};
+
+FCinematicProduction::FCinematicProduction()
+	: ProductionName(TEXT(""))
+	, DefaultDisplayRate(FFrameRate(24, 1))
+	, DefaultStartFrame(0)
+	, SubsequencePriority(ESubsequencePriority::BottomUp)
+{
+}
+
+FName UProductionSettings::GetContainerName() const
+{
+	return TEXT("Project");
+}
+
+FName UProductionSettings::GetCategoryName() const
+{
+	return TEXT("Plugins");
+}
+
+FName UProductionSettings::GetSectionName() const
+{
+	return TEXT("Production Settings");
+}
+
+void UProductionSettings::PostInitProperties()
+{
+	Super::PostInitProperties();
+
+	// Register a callback to detect when a ULevelSequenceProjectSettings property has changed to be notified when the DefaultDisplayRate changes
+	ULevelSequenceProjectSettings* LevelSequenceSettings = GetMutableDefault<ULevelSequenceProjectSettings>();
+	LevelSequenceSettings->OnSettingChanged().AddUObject(this, &UProductionSettings::OnLevelSequenceSettingChanged);
+
+	const UMovieSceneToolsProjectSettings* const MovieSceneToolsSettings = GetDefault<UMovieSceneToolsProjectSettings>();
+
+	// Cache the original tooltips for these properties to reset if the active production is set to "None"
+	FProperty* DefaultDisplayRateProperty = LevelSequenceSettings->GetClass()->FindPropertyByName(GET_MEMBER_NAME_CHECKED(ULevelSequenceProjectSettings, DefaultDisplayRate));
+	OriginalDefaultDisplayRateTooltip = DefaultDisplayRateProperty->GetMetaData(TEXT("ToolTip"));
+
+	FProperty* DefaultStartTimeProperty = MovieSceneToolsSettings->GetClass()->FindPropertyByName(GET_MEMBER_NAME_CHECKED(UMovieSceneToolsProjectSettings, DefaultStartTime));
+	OriginalDefaultStartTimeTooltip = DefaultStartTimeProperty->GetMetaData(TEXT("ToolTip"));
+
+	// Cache the original config value for the default display rate, so that it can be reset if the user sets their active production to "None"
+	DefaultDisplayRateProjectSetting = LevelSequenceSettings->DefaultDisplayRate;
+
+	// Load the serialized active production name from user settings
+	const FString ConfigSection = UProductionSettings::StaticClass()->GetPathName();
+	const FString PropertyName = GET_MEMBER_NAME_CHECKED(UProductionSettings, ActiveProductionName).ToString();
+
+	FString LoadedProductionName;
+	if (GConfig->GetString(*ConfigSection, *PropertyName, LoadedProductionName, GEditorPerProjectIni))
+	{
+		ActiveProductionName = *LoadedProductionName;
+	}
+	else
+	{
+		ActiveProductionName = NAME_None;
+	}
+
+	// Initialize the active production based on the serialized active production name
+	ActiveProduction = Productions.FindByPredicate([ProductionName = ActiveProductionName](const FCinematicProduction& Other) { return Other.ProductionName == ProductionName; });
+
+	// Apply the initial overrides to existing project settings based on the current active production
+	if (ActiveProduction)
+	{
+		ApplySequencerSettingOverrides();
+	}
+}
+
+void UProductionSettings::ApplySequencerSettingOverrides()
+{
+	// Modify the property flags and tooltip of these overridden sequencer settings
+	// They should be read-only when driven by an active production, and writable when the active production is None
+	ULevelSequenceProjectSettings* LevelSequenceSettings = GetMutableDefault<ULevelSequenceProjectSettings>();
+	FProperty* DefaultDisplayRateProperty = LevelSequenceSettings->GetClass()->FindPropertyByName(GET_MEMBER_NAME_CHECKED(ULevelSequenceProjectSettings, DefaultDisplayRate));
+ 
+	UMovieSceneToolsProjectSettings* MovieSceneToolsSettings = GetMutableDefault<UMovieSceneToolsProjectSettings>();
+ 	FProperty* DefaultStartTimeProperty = MovieSceneToolsSettings->GetClass()->FindPropertyByName(GET_MEMBER_NAME_CHECKED(UMovieSceneToolsProjectSettings, DefaultStartTime));
+
+	if (ActiveProduction)
+	{
+		DefaultDisplayRateProperty->SetPropertyFlags(EPropertyFlags::CPF_EditConst);
+		DefaultStartTimeProperty->SetPropertyFlags(EPropertyFlags::CPF_EditConst);
+
+		const FText ToolTipAddition = LOCTEXT("ToolTipAddition", "This property is being driven by the active production. To edit this value, change the production settings or set a different active production.");
+
+		DefaultDisplayRateProperty->SetMetaData(TEXT("ToolTip"), *FString::Printf(TEXT("%s\n\n%s"), *OriginalDefaultDisplayRateTooltip, *ToolTipAddition.ToString()));
+		DefaultStartTimeProperty->SetMetaData(TEXT("ToolTip"), *FString::Printf(TEXT("%s\n\n%s"), *OriginalDefaultStartTimeTooltip, *ToolTipAddition.ToString()));
+	}
+	else
+	{
+		DefaultDisplayRateProperty->ClearPropertyFlags(EPropertyFlags::CPF_EditConst);
+		DefaultStartTimeProperty->ClearPropertyFlags(EPropertyFlags::CPF_EditConst);
+
+		DefaultDisplayRateProperty->SetMetaData(TEXT("ToolTip"), *OriginalDefaultDisplayRateTooltip);
+		DefaultStartTimeProperty->SetMetaData(TEXT("ToolTip"), *OriginalDefaultStartTimeTooltip);
+	}
+
+	OverrideDefaultDisplayRate();
+	OverrideDefaultStartTime();
+	OverrideSubsequenceHierarchicalBias();
+}
+
+void UProductionSettings::OverrideDefaultDisplayRate()
+{
+	// Update the DefaultDisplayRate property of the level sequence project settings to the value of the active production
+	// If there is no active production, the level sequence project setting is reset to its default config value
+	ULevelSequenceProjectSettings* LevelSequenceSettings = GetMutableDefault<ULevelSequenceProjectSettings>();
+
+	if (ActiveProduction)
+	{
+		const FText FrameRateText = ActiveProduction->DefaultDisplayRate.ToPrettyText();
+		LevelSequenceSettings->DefaultDisplayRate = FrameRateText.ToString();
+	}
+	else
+	{
+		LevelSequenceSettings->DefaultDisplayRate = DefaultDisplayRateProjectSetting;
+	}
+
+	// Update the default display rate cvar to match
+	if (IConsoleVariable* ConsoleVariable = IConsoleManager::Get().FindConsoleVariable(TEXT("LevelSequence.DefaultDisplayRate")))
+	{
+		ConsoleVariable->Set(*LevelSequenceSettings->DefaultDisplayRate, ECVF_SetByProjectSetting);
+	}
+}
+
+void UProductionSettings::OverrideDefaultStartTime()
+{
+	if (ActiveProduction)
+	{
+		// Update the DefaultStartTime property of the movie scene tools project settings to the value of the active production
+		UMovieSceneToolsProjectSettings* MovieSceneToolsSettings = GetMutableDefault<UMovieSceneToolsProjectSettings>();
+
+		MovieSceneToolsSettings->DefaultStartTime = ActiveProduction->DefaultDisplayRate.AsSeconds(FFrameNumber(ActiveProduction->DefaultStartFrame));
+		MovieSceneToolsSettings->SaveConfig();
+	}
+}
+
+void UProductionSettings::OverrideSubsequenceHierarchicalBias()
+{
+	if (ActiveProduction)
+	{
+		using namespace UE::Shots::ProductionSettings::Private;
+		const int32 HierarchicalBiasValue = (ActiveProduction->SubsequencePriority == ESubsequencePriority::TopDown) ? TopDownHBiasValue : BottomUpHBiasValue;
+
+		// Write the new HBias value to the MovieScene SubSection section of the editor per project config file
+		const FString SubSectionClassPath = UMovieSceneSubSection::StaticClass()->GetPathName();
+		const FString ParametersPropertyName = GET_MEMBER_NAME_CHECKED(UMovieSceneSubSection, Parameters).ToString();
+		const FString HBiasPropertyName = GET_MEMBER_NAME_CHECKED(FMovieSceneSectionParameters, HierarchicalBias).ToString();
+
+		GConfig->SetString(*SubSectionClassPath, *ParametersPropertyName, *FString::Printf(TEXT("(%s=%d)"), *HBiasPropertyName, HierarchicalBiasValue), GEditorPerProjectIni);
+
+		// Update the CDO for the MovieScene SubSection so that new objects will use the new value 
+		UMovieSceneSubSection* SubSectionCDO = GetMutableDefault<UMovieSceneSubSection>();
+		SubSectionCDO->ReloadConfig();
+	}
+}
+
+void UProductionSettings::PostEditChangeProperty(struct FPropertyChangedEvent& PropertyChangedEvent)
+{
+	Super::PostEditChangeProperty(PropertyChangedEvent);
+
+	const FName PropertyName = PropertyChangedEvent.GetPropertyName();
+
+	if (PropertyName == GET_MEMBER_NAME_CHECKED(UProductionSettings, Productions))
+	{
+		if (PropertyChangedEvent.ChangeType == EPropertyChangeType::ArrayAdd)
+		{
+			// Make sure the production name is unique
+			const int32 ArrayIndex = PropertyChangedEvent.GetArrayIndex(PropertyChangedEvent.MemberProperty->GetFName().ToString());
+			FCinematicProduction& NewProduction = Productions[ArrayIndex];
+
+			const FName NameTemplate = TEXT("New Production");
+			FName WorkingName = NameTemplate;
+			int32 Suffix = 1;
+
+			while (Productions.ContainsByPredicate([WorkingName](const FCinematicProduction& Other) { return Other.ProductionName == WorkingName; }))
+			{
+				WorkingName = *(FString::Printf(TEXT("%s %d"), *NameTemplate.ToString(), Suffix));
+				Suffix++;
+			}
+
+			NewProduction.ProductionName = WorkingName;
+		}
+		else
+		{
+			// It is possible that a production was just deleted. If it was the active production, then that needs to be updated.
+			// This attempts to reacquire the pointer to the active production, and if it no longer exists, it will be reset to nullptr.
+			ActiveProduction = Productions.FindByPredicate([this](const FCinematicProduction& Other) { return Other.ProductionName == ActiveProductionName; });
+			SetActiveProductionName();
+
+			if (!ActiveProduction)
+			{
+				OverrideDefaultDisplayRate();
+			}
+		}
+	}
+	else if (PropertyName == GET_MEMBER_NAME_CHECKED(FCinematicProduction, ProductionName))
+	{
+		// If the active production's name property just changed, we need to update the serialized ActiveProductionName to match.
+		// Note: We don't know if the edited property belonged to the active production, but if it did not, the ActiveProductionName will remain unchanged.
+		SetActiveProductionName();
+	}
+	else if (PropertyName == GET_MEMBER_NAME_CHECKED(FCinematicProduction, DefaultDisplayRate))
+	{
+		OverrideDefaultDisplayRate();
+		OverrideDefaultStartTime();
+	}
+	else if (PropertyName == GET_MEMBER_NAME_CHECKED(FCinematicProduction, DefaultStartFrame))
+	{
+		OverrideDefaultStartTime();
+	}
+	else if (PropertyName == GET_MEMBER_NAME_CHECKED(FCinematicProduction, SubsequencePriority))
+	{
+		OverrideSubsequenceHierarchicalBias();
+	}
+}
+
+void UProductionSettings::OnLevelSequenceSettingChanged(UObject* Object, struct FPropertyChangedEvent& PropertyChangedEvent)
+{
+	if (PropertyChangedEvent.GetPropertyName() == GET_MEMBER_NAME_CHECKED(ULevelSequenceProjectSettings, DefaultDisplayRate))
+	{
+		const ULevelSequenceProjectSettings* const LevelSequenceSettings = GetDefault<ULevelSequenceProjectSettings>();
+		DefaultDisplayRateProjectSetting = LevelSequenceSettings->DefaultDisplayRate;
+	}
+}
+
+void UProductionSettings::SetActiveProduction(FName ProductionName)
+{
+	// Find and set the active production based on the input production name
+	ActiveProduction = Productions.FindByPredicate([ProductionName](const FCinematicProduction& Other) { return Other.ProductionName == ProductionName; });
+	SetActiveProductionName();
+
+	// Apply the overrides based on the new active production (even if it is None, because some properties will be reset to default)
+	ApplySequencerSettingOverrides();
+}
+
+void UProductionSettings::SetActiveProductionName()
+{
+	ActiveProductionName = ActiveProduction ? ActiveProduction->ProductionName : NAME_None;
+
+	const FString ConfigSection = UProductionSettings::StaticClass()->GetPathName();
+	const FString PropertyName = GET_MEMBER_NAME_CHECKED(UProductionSettings, ActiveProductionName).ToString();
+
+	GConfig->SetString(*ConfigSection, *PropertyName, *ActiveProductionName.ToString(), GEditorPerProjectIni);
+	GConfig->Flush(false, GEditorPerProjectIni);
+}
+
+#undef LOCTEXT_NAMESPACE
\ No newline at end of file
diff --git a/Engine/Plugins/VirtualProduction/Shots/Source/ShotsEditor/Private/ProductionSettings.h b/Engine/Plugins/VirtualProduction/Shots/Source/ShotsEditor/Private/ProductionSettings.h
new file mode 100644
--- /dev/null
+++ b/Engine/Plugins/VirtualProduction/Shots/Source/ShotsEditor/Private/ProductionSettings.h
@@ -0,0 +1,110 @@
+// Copyright Epic Games, Inc. All Rights Reserved.
+#pragma  once
+
+#include "Engine/DeveloperSettings.h"
+
+#include "Misc/FrameRate.h"
+
+#include "ProductionSettings.generated.h"
+
+/**
+ * Options for determining the hierarchical bias of subsequences
+ */
+UENUM()
+enum class ESubsequencePriority : uint8
+{
+	TopDown,
+	BottomUp
+};
+
+/**
+ * Collection of production settings to override the project/editor behavior
+ */
+USTRUCT()
+struct FCinematicProduction
+{
+	GENERATED_BODY()
+
+public:
+	FCinematicProduction();
+
+	UPROPERTY(config, EditAnywhere, Category = "Default")
+	FName ProductionName;
+
+	/** The default frame rate to set for new Level Sequences */
+	UPROPERTY(config, EditAnywhere, Category = "Default")
+	FFrameRate DefaultDisplayRate;
+
+	/** The default frame number (using the default frame rate) that new Level Sequences should start at */
+	UPROPERTY(config, EditAnywhere, Category = "Default")
+	int32 DefaultStartFrame;
+
+	/** Controls whether subsequences override parent sequences, or vice versa */
+	UPROPERTY(config, EditAnywhere, Category = "Default")
+	ESubsequencePriority SubsequencePriority = ESubsequencePriority::BottomUp;
+};
+
+/**
+ * Cinematic Production Settings
+ */
+UCLASS(config=Engine, defaultconfig)
+class UProductionSettings : public UDeveloperSettings
+{
+	GENERATED_BODY()
+
+public:
+	UProductionSettings() = default;
+
+	//~ Begin UDeveloperSettings overrides
+	virtual FName GetContainerName() const override;
+	virtual FName GetCategoryName() const override;
+	virtual FName GetSectionName() const override;
+	//~ End UDeveloperSettings overrides
+
+	//~ Begin UObject overrides
+	virtual void PostInitProperties() override;
+	virtual void PostEditChangeProperty(struct FPropertyChangedEvent& PropertyChangedEvent) override;
+	//~ End UObject overrides
+
+	/** Sets the active production based on the input production name */
+	void SetActiveProduction(FName ProductionName);
+
+private:
+	/** Callback when a property of ULevelSequenceProjectSettings has changed */
+	void OnLevelSequenceSettingChanged(UObject* Object, struct FPropertyChangedEvent& PropertyChangedEvent);
+
+	/** Applies overrides to various sequencer settings based on the active production settings */
+	void ApplySequencerSettingOverrides();
+
+	/** Overrides the DefaultDisplayRate in the level sequence project settings based on the active production setting */
+	void OverrideDefaultDisplayRate();
+
+	/** Overrides the DefaultStartTime in the movie scene tools project settings based on the active production settings */
+	void OverrideDefaultStartTime();
+
+	/** Writes out the Hierarchical bias value (based on the Subsequence Priority setting) to EditorPerProjectUserSettings.ini */
+	void OverrideSubsequenceHierarchicalBias();
+
+	/** Set the active production name and write out the new value to EditorPerProjectUserSettings.ini */
+	void SetActiveProductionName();
+
+public:
+	/** Name of the active production */
+	UPROPERTY(VisibleAnywhere, Category = "Default", meta=(DisplayName="Active Production"))
+	FName ActiveProductionName;
+
+	/** List of available productions in this project */
+	UPROPERTY(config, EditAnywhere, Category = "Default")
+	TArray<FCinematicProduction> Productions;
+
+private:
+	/** Pointer to the active production (in the Productions array) */
+	FCinematicProduction* ActiveProduction;
+
+	/** Default value of DefaultDisplayRate in the level sequence project setting, used to reset it when there is no active production */
+	FString DefaultDisplayRateProjectSetting;
+
+	/** Original tooltip text for sequencer settings */
+	FString OriginalDefaultDisplayRateTooltip;
+	FString OriginalDefaultStartTimeTooltip;
+};
diff --git a/Engine/Plugins/VirtualProduction/Shots/Source/ShotsEditor/Private/ProductionSettingsCustomization.cpp b/Engine/Plugins/VirtualProduction/Shots/Source/ShotsEditor/Private/ProductionSettingsCustomization.cpp
new file mode 100644
--- /dev/null
+++ b/Engine/Plugins/VirtualProduction/Shots/Source/ShotsEditor/Private/ProductionSettingsCustomization.cpp
@@ -0,0 +1,112 @@
+// Copyright Epic Games, Inc. All Rights Reserved.
+
+#include "ProductionSettingsCustomization.h"
+
+#include "DetailLayoutBuilder.h"
+#include "DetailWidgetRow.h"
+#include "Framework/MultiBox/MultiBoxBuilder.h"
+#include "Modules/ModuleManager.h"
+#include "ProductionSettings.h"
+#include "Widgets/Input/SComboButton.h"
+
+TSharedRef<IDetailCustomization> FProductionSettingsCustomization::MakeInstance()
+{
+	return MakeShareable(new FProductionSettingsCustomization);
+}
+
+void FProductionSettingsCustomization::CustomizeDetails(IDetailLayoutBuilder& DetailBuilder)
+{
+	TArray<TWeakObjectPtr<UObject>> CustomizedObjects;
+	DetailBuilder.GetObjectsBeingCustomized(CustomizedObjects);
+
+	PropertyUtilities = DetailBuilder.GetPropertyUtilities();
+
+	// Ensure that we are only customizing one object
+	if (CustomizedObjects.Num() != 1)
+	{
+		return;
+	}
+
+	// Ensure that we are only customizing one object
+	UProductionSettings* CustomizedProductionSettings = Cast<UProductionSettings>(CustomizedObjects[0]);
+	if (CustomizedProductionSettings != GetDefault<UProductionSettings>())
+	{
+		return;
+	}
+
+	// Hide the default active production property, and replace with a combo button listing the available productions in this project
+ 	TSharedRef<IPropertyHandle> ActiveProductionNameHandle = DetailBuilder.GetProperty(GET_MEMBER_NAME_CHECKED(UProductionSettings, ActiveProductionName));
+ 	DetailBuilder.HideProperty(ActiveProductionNameHandle);
+
+ 	DetailBuilder.AddCustomRowToCategory(ActiveProductionNameHandle, ActiveProductionNameHandle->GetPropertyDisplayName())
+		.NameContent()
+		[
+			ActiveProductionNameHandle->CreatePropertyNameWidget()
+		]
+		.ValueContent()
+		[
+			SNew(SComboButton)
+				.VAlign(VAlign_Center)
+				.OnGetMenuContent(this, &FProductionSettingsCustomization::BuildProductionNameMenu)
+				.ButtonContent()
+				[
+					SNew(STextBlock).Text(this, &FProductionSettingsCustomization::GetActiveProductionName)
+						.Font(IDetailLayoutBuilder::GetDetailFont())
+				]
+		];
+}
+
+TSharedRef<SWidget> FProductionSettingsCustomization::BuildProductionNameMenu()
+{
+	FMenuBuilder MenuBuilder(true, nullptr);
+
+	// Always add a "None" option to set the active production to None
+	FName NoActiveProductionName = NAME_None;
+	MenuBuilder.AddMenuEntry(
+		FText::FromName(NoActiveProductionName),
+		FText::FromName(NoActiveProductionName),
+		FSlateIcon(),
+		FUIAction(FExecuteAction::CreateRaw(this, &FProductionSettingsCustomization::SetActiveProduction, NoActiveProductionName)),
+		NAME_None,
+		EUserInterfaceActionType::None
+	);
+
+	const UProductionSettings* const ProductionSettings = GetDefault<UProductionSettings>();
+
+	// Add a menu option with the production name for each production available in this project
+	TArray<FName> ProductionNames;
+	for (const FCinematicProduction& Production : ProductionSettings->Productions)
+	{
+		ProductionNames.Add(Production.ProductionName);
+	}
+
+	for (FName Name : ProductionNames)
+	{
+		MenuBuilder.AddMenuEntry(
+			FText::FromName(Name),
+			FText::FromName(Name),
+			FSlateIcon(),
+			FUIAction(FExecuteAction::CreateRaw(this, &FProductionSettingsCustomization::SetActiveProduction, Name)),
+			NAME_None,
+			EUserInterfaceActionType::None
+		);
+	}
+
+	return MenuBuilder.MakeWidget();
+}
+
+FText FProductionSettingsCustomization::GetActiveProductionName() const
+{
+ 	const UProductionSettings* const ProductionSettings = GetDefault<UProductionSettings>();
+	return FText::FromName(ProductionSettings->ActiveProductionName);
+}
+
+void FProductionSettingsCustomization::SetActiveProduction(FName ProductionName)
+{
+	UProductionSettings* ProductionSettings = GetMutableDefault<UProductionSettings>();
+	ProductionSettings->SetActiveProduction(ProductionName);
+
+	// The active production influences whether certain sequencer settings are writable or read-only.
+	// This forces the details view of the other sequencer settings categories to refresh immediately to respect the new property flags.
+	PropertyUtilities->RequestForceRefresh();
+}
diff --git a/Engine/Plugins/VirtualProduction/Shots/Source/ShotsEditor/Private/ProductionSettingsCustomization.h b/Engine/Plugins/VirtualProduction/Shots/Source/ShotsEditor/Private/ProductionSettingsCustomization.h
new file mode 100644
--- /dev/null
+++ b/Engine/Plugins/VirtualProduction/Shots/Source/ShotsEditor/Private/ProductionSettingsCustomization.h
@@ -0,0 +1,37 @@
+// Copyright Epic Games, Inc. All Rights Reserved.
+
+#pragma once
+
+#include "IDetailCustomization.h"
+
+#include "IPropertyUtilities.h"
+
+class IDetailLayoutBuilder;
+class SWidget;
+
+/**
+ * Detail customization for UProductionSettings
+ */
+class FProductionSettingsCustomization : public IDetailCustomization
+{
+public:
+	/** Makes a new instance of this detail layout class for a specific detail view requesting it */
+	static TSharedRef<IDetailCustomization> MakeInstance();
+
+	//~ Begin IDetailCustomization interface */
+	virtual void CustomizeDetails(IDetailLayoutBuilder& DetailBuilder) override;
+	//~ End IDetailCustomization interface */
+
+private:
+	/** Set the active production of the production project settings */
+	void SetActiveProduction(FName ProductionName);
+
+	/** Get the name of the currently active production */
+	FText GetActiveProductionName() const;
+
+	/** Build the menu widget for the production name combo button */
+	TSharedRef<SWidget> BuildProductionNameMenu();
+
+	/** Property Utils, used to refresh the property settings details when changing the active production */
+	TSharedPtr<IPropertyUtilities> PropertyUtilities;
+};
diff --git a/Engine/Plugins/VirtualProduction/Shots/Source/ShotsEditor/Private/ShotsEditorModule.cpp b/Engine/Plugins/VirtualProduction/Shots/Source/ShotsEditor/Private/ShotsEditorModule.cpp
--- a/Engine/Plugins/VirtualProduction/Shots/Source/ShotsEditor/Private/ShotsEditorModule.cpp
+++ b/Engine/Plugins/VirtualProduction/Shots/Source/ShotsEditor/Private/ShotsEditorModule.cpp
@@ -2,4 +2,34 @@
 
 #include "Modules/ModuleManager.h"
 
-IMPLEMENT_MODULE(FDefaultModuleImpl, ShotsEditor);
+#include "ProductionSettings.h"
+#include "ProductionSettingsCustomization.h"
+#include "PropertyEditorModule.h"
+
+class FShotsEditorModule : public IModuleInterface
+{
+public:
+
+	//~ Begin IModuleInterface interface
+	virtual void StartupModule() override;
+	virtual void ShutdownModule() override;
+	//~ End IModuleInterface interface
+};
+
+void FShotsEditorModule::StartupModule()
+{
+	FPropertyEditorModule& PropertyEditorModule = FModuleManager::Get().LoadModuleChecked<FPropertyEditorModule>("PropertyEditor");
+
+	PropertyEditorModule.RegisterCustomClassLayout(
+		UProductionSettings::StaticClass()->GetFName(),
+		FOnGetDetailCustomizationInstance::CreateStatic(&FProductionSettingsCustomization::MakeInstance)
+	);
+}
+
+void FShotsEditorModule::ShutdownModule()
+{
+	FPropertyEditorModule& PropertyModule = FModuleManager::GetModuleChecked<FPropertyEditorModule>("PropertyEditor");
+	PropertyModule.UnregisterCustomClassLayout(UProductionSettings::StaticClass()->GetFName());
+}
+
+IMPLEMENT_MODULE(FShotsEditorModule, ShotsEditor);
diff --git a/Engine/Plugins/VirtualProduction/Shots/Source/ShotsEditor/ShotsEditor.build.cs b/Engine/Plugins/VirtualProduction/Shots/Source/ShotsEditor/ShotsEditor.build.cs
--- a/Engine/Plugins/VirtualProduction/Shots/Source/ShotsEditor/ShotsEditor.build.cs
+++ b/Engine/Plugins/VirtualProduction/Shots/Source/ShotsEditor/ShotsEditor.build.cs
@@ -20,6 +20,16 @@ namespace UnrealBuildTool.Rules
 					"ContentBrowserData",
 					"Core",
 					"CoreUObject",
+					"DeveloperSettings",
+					"Engine",
+					"InputCore",
+					"LevelSequence",
+					"LevelSequenceEditor",
+					"MovieScene",
+					"MovieSceneTools",
+					"Slate",
+					"SlateCore",
+					"SourceControl",
 					"UnrealEd",
 				}
 			);
diff --git a/Engine/Source/Runtime/LevelSequence/Private/LevelSequenceProjectSettings.h b/Engine/Source/Runtime/LevelSequence/Private/LevelSequenceProjectSettings.h
deleted file mode 100644
diff --git a/Engine/Source/Runtime/LevelSequence/Public/LevelSequenceProjectSettings.h b/Engine/Source/Runtime/LevelSequence/Public/LevelSequenceProjectSettings.h
new file mode 100644
--- /dev/null
+++ b/Engine/Source/Runtime/LevelSequence/Public/LevelSequenceProjectSettings.h
@@ -0,0 +1,38 @@
+// Copyright Epic Games, Inc. All Rights Reserved.
+
+#pragma once
+
+#include "CoreMinimal.h"
+#include "UObject/ObjectMacros.h"
+#include "UObject/Object.h"
+#include "Engine/DeveloperSettings.h"
+#include "MovieSceneFwd.h"
+#include "LevelSequenceProjectSettings.generated.h"
+
+// Settings for level sequences
+UCLASS(config=Engine, defaultconfig, meta=(DisplayName="Level Sequence"))
+class LEVELSEQUENCE_API ULevelSequenceProjectSettings : public UDeveloperSettings
+{
+	GENERATED_BODY()
+
+public:
+	ULevelSequenceProjectSettings();
+
+	UPROPERTY(config, EditAnywhere, Category=Timeline, meta=(ConsoleVariable="LevelSequence.DefaultLockEngineToDisplayRate", Tooltip="0: Playback locked to playback frames\n1: Unlocked playback with sub frame interpolation"))
+	bool bDefaultLockEngineToDisplayRate;
+
+	UPROPERTY(config, EditAnywhere, Category=Timeline, meta=(ConsoleVariable="LevelSequence.DefaultDisplayRate", Tooltip="Specifies default display frame rate for newly created level sequences; also defines frame locked frame rate where sequences are set to be frame locked. Examples: 30 fps, 120/1 (120 fps), 30000/1001 (29.97), 0.01s (10ms)."))
+	FString DefaultDisplayRate;
+
+	UPROPERTY(config, EditAnywhere, Category=Timeline, meta=(ConsoleVariable="LevelSequence.DefaultTickResolution", Tooltip="Specifies default tick resolution for newly created level sequences. Examples: 30 fps, 120/1 (120 fps), 30000/1001 (29.97), 0.01s (10ms)."))
+	FString DefaultTickResolution;
+
+	UPROPERTY(config, EditAnywhere, Category=Timeline, meta=(ConsoleVariable="LevelSequence.DefaultClockSource", Tooltip="Specifies default clock source for newly created level sequences. Examples: 0: Tick, 1: Platform, 2: Audio, 3: RelativeTimecode, 4: Timecode, 5: Custom"))
+	EUpdateClockSource DefaultClockSource;
+
+	virtual void PostInitProperties() override;
+#if WITH_EDITOR
+	virtual void PostEditChangeProperty(struct FPropertyChangedEvent& PropertyChangedEvent) override;
+#endif
+};
+
