ShotsManagement: Add framework for the production setup wizard tool, along with the Productions and Sequencer Settings panels.
The ShotsEditor module now registers the Production Wizard as a new tab that can be found in Windows->Cinematics.
The Production Settings interface has been filled out more to support the production wizard workflows.
The underlying string type for productions was swapped from FName to FString, which is more correct for user-entered strings.
This CL also introduces a style set for this plugin.

diff --git a/Engine/Plugins/VirtualProduction/Shots/Source/ShotsEditor/Private/ProductionSettings.cpp b/Engine/Plugins/VirtualProduction/Shots/Source/ShotsEditor/Private/ProductionSettings.cpp
--- a/Engine/Plugins/VirtualProduction/Shots/Source/ShotsEditor/Private/ProductionSettings.cpp
+++ b/Engine/Plugins/VirtualProduction/Shots/Source/ShotsEditor/Private/ProductionSettings.cpp
@@ -20,11 +20,15 @@ namespace UE::Shots::ProductionSettings::Private
 };
 
 FCinematicProduction::FCinematicProduction()
-	: ProductionName(TEXT(""))
+	: ProductionName()
 	, DefaultDisplayRate(FFrameRate(24, 1))
 	, DefaultStartFrame(0)
 	, SubsequencePriority(ESubsequencePriority::BottomUp)
 {
+	if (!ProductionID.IsValid())
+	{
+		ProductionID = FGuid::NewGuid();
+	}
 }
 
 FName UProductionSettings::GetContainerName() const
@@ -69,11 +73,11 @@ void UProductionSettings::PostInitProperties()
 	FString LoadedProductionName;
 	if (GConfig->GetString(*ConfigSection, *PropertyName, LoadedProductionName, GEditorPerProjectIni))
 	{
-		ActiveProductionName = *LoadedProductionName;
+		ActiveProductionName = LoadedProductionName;
 	}
 	else
 	{
-		ActiveProductionName = NAME_None;
+		ActiveProductionName = FString(TEXT(""));
 	}
 
 	// Initialize the active production based on the serialized active production name
@@ -168,6 +172,7 @@ void UProductionSettings::OverrideSubsequenceHierarchicalBias()
 		const FString HBiasPropertyName = GET_MEMBER_NAME_CHECKED(FMovieSceneSectionParameters, HierarchicalBias).ToString();
 
 		GConfig->SetString(*SubSectionClassPath, *ParametersPropertyName, *FString::Printf(TEXT("(%s=%d)"), *HBiasPropertyName, HierarchicalBiasValue), GEditorPerProjectIni);
+		GConfig->Flush(false, GEditorPerProjectIni);
 
 		// Update the CDO for the MovieScene SubSection so that new objects will use the new value 
 		UMovieSceneSubSection* SubSectionCDO = GetMutableDefault<UMovieSceneSubSection>();
@@ -187,19 +192,7 @@ void UProductionSettings::PostEditChangeProperty(struct FPropertyChangedEvent& P
 		{
 			// Make sure the production name is unique
 			const int32 ArrayIndex = PropertyChangedEvent.GetArrayIndex(PropertyChangedEvent.MemberProperty->GetFName().ToString());
-			FCinematicProduction& NewProduction = Productions[ArrayIndex];
-
-			const FName NameTemplate = TEXT("New Production");
-			FName WorkingName = NameTemplate;
-			int32 Suffix = 1;
-
-			while (Productions.ContainsByPredicate([WorkingName](const FCinematicProduction& Other) { return Other.ProductionName == WorkingName; }))
-			{
-				WorkingName = *(FString::Printf(TEXT("%s %d"), *NameTemplate.ToString(), Suffix));
-				Suffix++;
-			}
-
-			NewProduction.ProductionName = WorkingName;
+			Productions[ArrayIndex].ProductionName = GetUniqueProductionName();
 		}
 		else
 		{
@@ -213,6 +206,8 @@ void UProductionSettings::PostEditChangeProperty(struct FPropertyChangedEvent& P
 				OverrideDefaultDisplayRate();
 			}
 		}
+
+		ProductionListChangedDelegate.Broadcast();
 	}
 	else if (PropertyName == GET_MEMBER_NAME_CHECKED(FCinematicProduction, ProductionName))
 	{
@@ -244,25 +239,237 @@ void UProductionSettings::OnLevelSequenceSettingChanged(UObject* Object, struct
 	}
 }
 
-void UProductionSettings::SetActiveProduction(FName ProductionName)
+void UProductionSettings::AddProduction()
+{
+	const int32 ArrayIndex = Productions.Add(FCinematicProduction());
+
+	TryUpdateDefaultConfigFile();
+
+	ProductionListChangedDelegate.Broadcast();
+}
+
+void UProductionSettings::DeleteProduction(FGuid ProductionID)
+{
+	if (!ProductionID.IsValid())
+	{
+		return;
+	}
+
+	int32 IndexToDelete = Productions.IndexOfByPredicate([ProductionID](const FCinematicProduction& Other) { return Other.ProductionID == ProductionID; });
+
+	if (IndexToDelete != INDEX_NONE)
+	{
+		// Determine whether this was the active production before removing it
+		const bool bIsActiveProduction = IsActiveProduction(ProductionID);
+
+		Productions.RemoveAt(IndexToDelete);
+
+		if (bIsActiveProduction)
+		{
+			ActiveProduction = nullptr;
+			SetActiveProductionName();
+
+			// This will restore the level sequence project setting to its original default setting (before any production overrode it)
+			OverrideDefaultDisplayRate();
+		}
+
+		TryUpdateDefaultConfigFile();
+
+		ProductionListChangedDelegate.Broadcast();
+	}
+}
+
+void UProductionSettings::RenameProduction(FGuid ProductionID, FString NewName)
+{
+	FCinematicProduction* Production = Productions.FindByPredicate([ProductionID](const FCinematicProduction& Other) { return Other.ProductionID == ProductionID; });
+	if (Production)
+	{
+		Production->ProductionName = NewName;
+		SetActiveProductionName();
+		TryUpdateDefaultConfigFile();
+	}
+}
+
+const FCinematicProduction* UProductionSettings::GetProduction(FGuid ProductionID) const
+{
+	if (!ProductionID.IsValid())
+	{
+		return nullptr;
+	}
+
+	return Productions.FindByPredicate([ProductionID](const FCinematicProduction& Other) { return Other.ProductionID == ProductionID; });
+}
+
+FCinematicProduction* UProductionSettings::GetActiveProduction() const
 {
-	// Find and set the active production based on the input production name
-	ActiveProduction = Productions.FindByPredicate([ProductionName](const FCinematicProduction& Other) { return Other.ProductionName == ProductionName; });
+	return ActiveProduction;
+}
+
+void UProductionSettings::SetActiveProduction(FGuid ProductionID)
+{
+	ActiveProduction = Productions.FindByPredicate([ProductionID](const FCinematicProduction& Other) { return Other.ProductionID == ProductionID; });
 	SetActiveProductionName();
 
 	// Apply the overrides based on the new active production (even if it is None, because some properties will be reset to default)
 	ApplySequencerSettingOverrides();
 }
 
+bool UProductionSettings::IsActiveProduction(FGuid ProductionID) const
+{
+	if (ActiveProduction && ProductionID.IsValid())
+	{
+		return (ProductionID == ActiveProduction->ProductionID);
+	}
+	return false;
+}
+
+FFrameRate UProductionSettings::GetActiveDisplayRate() const
+{
+	if (ActiveProduction)
+	{
+		return ActiveProduction->DefaultDisplayRate;
+	}
+
+	// If there is no active production, return the current level sequence project setting instead
+	const ULevelSequenceProjectSettings* const LevelSequenceSettings = GetDefault<ULevelSequenceProjectSettings>();
+
+	FFrameRate DefaultSetting;
+	TryParseString(DefaultSetting, *LevelSequenceSettings->DefaultDisplayRate);
+
+	return DefaultSetting;
+}
+
+int32 UProductionSettings::GetActiveStartFrame() const
+{
+	if (ActiveProduction)
+	{
+		return ActiveProduction->DefaultStartFrame;
+	}
+
+	// If there is no active production, compute the equivalent frame number from the current movie scene tools project setting instead
+	const UMovieSceneToolsProjectSettings* const MovieSceneToolsSettings = GetDefault<UMovieSceneToolsProjectSettings>();
+	const double StartTimeInSeconds = MovieSceneToolsSettings->DefaultStartTime;
+
+	FFrameRate DefaultFrameRate = GetActiveDisplayRate();
+	return DefaultFrameRate.AsFrameNumber(StartTimeInSeconds).Value;
+}
+
+ESubsequencePriority UProductionSettings::GetActiveSubsequencePriority() const
+{
+	using namespace UE::Shots::ProductionSettings::Private;
+
+	if (ActiveProduction)
+	{
+		return ActiveProduction->SubsequencePriority;
+	}
+
+	// If there is no active production, return the value based on the current movie scene subsection setting
+	const UMovieSceneSubSection* const SubSectionCDO = GetDefault<UMovieSceneSubSection>();
+
+	if (SubSectionCDO->Parameters.HierarchicalBias == TopDownHBiasValue)
+	{
+		return ESubsequencePriority::TopDown;
+	}
+	return ESubsequencePriority::BottomUp;
+}
+
+void UProductionSettings::SetDisplayRate(FGuid ProductionID, FFrameRate DisplayRate)
+{
+	FCinematicProduction* Production = Productions.FindByPredicate([ProductionID](const FCinematicProduction& Other) { return Other.ProductionID == ProductionID; });
+	if (Production)
+	{
+		Production->DefaultDisplayRate = DisplayRate;
+
+		if (IsActiveProduction(ProductionID))
+		{
+			OverrideDefaultDisplayRate();
+			OverrideDefaultStartTime();
+		}
+
+		TryUpdateDefaultConfigFile();
+	}
+}
+
+void UProductionSettings::SetStartFrame(FGuid ProductionID, int32 StartFrame)
+{
+	FCinematicProduction* Production = Productions.FindByPredicate([ProductionID](const FCinematicProduction& Other) { return Other.ProductionID == ProductionID; });
+	if (Production)
+	{
+		Production->DefaultStartFrame = StartFrame;
+
+		if (IsActiveProduction(ProductionID))
+		{
+			OverrideDefaultStartTime();
+		}
+
+		TryUpdateDefaultConfigFile();
+	}
+}
+
+void UProductionSettings::SetSubsequencePriority(FGuid ProductionID, ESubsequencePriority Priority)
+{
+	FCinematicProduction* Production = Productions.FindByPredicate([ProductionID](const FCinematicProduction& Other) { return Other.ProductionID == ProductionID; });
+	if (Production)
+	{
+		Production->SubsequencePriority = Priority;
+
+		if (IsActiveProduction(ProductionID))
+		{
+			OverrideSubsequenceHierarchicalBias();
+		}
+
+		TryUpdateDefaultConfigFile();
+	}
+}
+
+void UProductionSettings::SetActiveDisplayRate(FFrameRate DisplayRate)
+{
+	if (ActiveProduction)
+	{
+		SetDisplayRate(ActiveProduction->ProductionID, DisplayRate);
+	}
+}
+
+void UProductionSettings::SetActiveStartFrame(int32 StartFrame)
+{
+	if (ActiveProduction)
+	{
+		SetStartFrame(ActiveProduction->ProductionID, StartFrame);
+	}
+}
+
+void UProductionSettings::SetActiveSubsequencePriority(ESubsequencePriority Priority)
+{
+	if (ActiveProduction)
+	{
+		SetSubsequencePriority(ActiveProduction->ProductionID, Priority);
+	}
+}
+
 void UProductionSettings::SetActiveProductionName()
 {
-	ActiveProductionName = ActiveProduction ? ActiveProduction->ProductionName : NAME_None;
+	ActiveProductionName = ActiveProduction ? ActiveProduction->ProductionName : FString(TEXT(""));
 
 	const FString ConfigSection = UProductionSettings::StaticClass()->GetPathName();
 	const FString PropertyName = GET_MEMBER_NAME_CHECKED(UProductionSettings, ActiveProductionName).ToString();
 
-	GConfig->SetString(*ConfigSection, *PropertyName, *ActiveProductionName.ToString(), GEditorPerProjectIni);
+	GConfig->SetString(*ConfigSection, *PropertyName, *ActiveProductionName, GEditorPerProjectIni);
 	GConfig->Flush(false, GEditorPerProjectIni);
 }
 
+FString UProductionSettings::GetUniqueProductionName() const
+{
+	const FText NameTemplate = LOCTEXT("NewProductionText", "New Production");
+	FString WorkingName = NameTemplate.ToString();
+	int32 Suffix = 1;
+
+	while (Productions.ContainsByPredicate([WorkingName](const FCinematicProduction& Other) { return Other.ProductionName == WorkingName; }))
+	{
+		WorkingName = FString::Printf(TEXT("%s %d"), *NameTemplate.ToString(), Suffix);
+		Suffix++;
+	}
+
+	return WorkingName;
+}
+
 #undef LOCTEXT_NAMESPACE
\ No newline at end of file
diff --git a/Engine/Plugins/VirtualProduction/Shots/Source/ShotsEditor/Private/ProductionSettings.h b/Engine/Plugins/VirtualProduction/Shots/Source/ShotsEditor/Private/ProductionSettings.h
--- a/Engine/Plugins/VirtualProduction/Shots/Source/ShotsEditor/Private/ProductionSettings.h
+++ b/Engine/Plugins/VirtualProduction/Shots/Source/ShotsEditor/Private/ProductionSettings.h
@@ -7,6 +7,8 @@
 
 #include "ProductionSettings.generated.h"
 
+DECLARE_MULTICAST_DELEGATE(FOnProductionListChanged);
+
 /**
  * Options for determining the hierarchical bias of subsequences
  */
@@ -29,7 +31,7 @@ public:
 	FCinematicProduction();
 
 	UPROPERTY(config, EditAnywhere, Category = "Default")
-	FName ProductionName;
+	FString ProductionName;
 
 	/** The default frame rate to set for new Level Sequences */
 	UPROPERTY(config, EditAnywhere, Category = "Default")
@@ -42,6 +44,10 @@ public:
 	/** Controls whether subsequences override parent sequences, or vice versa */
 	UPROPERTY(config, EditAnywhere, Category = "Default")
 	ESubsequencePriority SubsequencePriority = ESubsequencePriority::BottomUp;
+
+	/** Unique ID of the production */
+	UPROPERTY(meta=(IgnoreForMemberInitializationTest))
+	FGuid ProductionID;
 };
 
 /**
@@ -66,8 +72,59 @@ public:
 	virtual void PostEditChangeProperty(struct FPropertyChangedEvent& PropertyChangedEvent) override;
 	//~ End UObject overrides
 
-	/** Sets the active production based on the input production name */
-	void SetActiveProduction(FName ProductionName);
+	/** Adds a new production to the list */
+	void AddProduction();
+
+	/** Removes the production matching the input ID from the list */
+	void DeleteProduction(FGuid ProductionID);
+
+	/** Renames the production matching the input ID */
+	void RenameProduction(FGuid ProductionID, FString NewName);
+
+	/** Returns a pointer to the production matching the input ID */
+	const FCinematicProduction* GetProduction(FGuid ProductionID) const;
+
+	/** Returns a pointer to the active production */
+	FCinematicProduction* GetActiveProduction() const;
+
+	/** Sets the active production based on the input production ID */
+	void SetActiveProduction(FGuid ProductionID);
+
+	/** Returns true if the input production ID matches the ID of the active production */
+	bool IsActiveProduction(FGuid ProductionID) const;
+
+	/** Returns the DefaultDisplayRate of the active production, or the underlying level sequence setting if there is no active production */
+	FFrameRate GetActiveDisplayRate() const;
+
+	/** Returns the DefaultStartFrame of the active production, or the value based on the underlying movie scene tools setting if there is no active production */
+	int32 GetActiveStartFrame() const;
+
+	/** Returns the SubsequencePriority of the active production, or the default config value if there is no active production */
+	ESubsequencePriority GetActiveSubsequencePriority() const;
+
+	/** Sets the DefaultDisplayRate of the production matching the input ID */
+	void SetDisplayRate(FGuid ProductionID, FFrameRate DisplayRate);
+
+	/** Sets the DefaultStartFrame of the production matching the input ID */
+	void SetStartFrame(FGuid ProductionID, int32 StartFrame);
+
+	/** Sets the SubsequencePriority of the production matching the input ID */
+	void SetSubsequencePriority(FGuid ProductionID, ESubsequencePriority Priority);
+
+	/** Sets the DefaultDisplayRate of the active production */
+	void SetActiveDisplayRate(FFrameRate DisplayRate);
+
+	/** Sets the DefaultStartFrame of the active production */
+	void SetActiveStartFrame(int32 StartFrame);
+
+	/** Sets the SubsequencePriority of the active production */
+	void SetActiveSubsequencePriority(ESubsequencePriority Priority);
+
+	/** Returns a new unique production name */
+	FString GetUniqueProductionName() const;
+
+	/** Returns the delegate that broadcasts when a production is added/removed */
+	FOnProductionListChanged& OnProductionListChanged() { return ProductionListChangedDelegate; }
 
 private:
 	/** Callback when a property of ULevelSequenceProjectSettings has changed */
@@ -91,7 +148,7 @@ private:
 public:
 	/** Name of the active production */
 	UPROPERTY(VisibleAnywhere, Category = "Default", meta=(DisplayName="Active Production"))
-	FName ActiveProductionName;
+	FString ActiveProductionName;
 
 	/** List of available productions in this project */
 	UPROPERTY(config, EditAnywhere, Category = "Default")
@@ -107,4 +164,7 @@ private:
 	/** Original tooltip text for sequencer settings */
 	FString OriginalDefaultDisplayRateTooltip;
 	FString OriginalDefaultStartTimeTooltip;
+
+	/** Delegate that broadcasts when a production is added/removed */
+	FOnProductionListChanged ProductionListChangedDelegate;
 };
diff --git a/Engine/Plugins/VirtualProduction/Shots/Source/ShotsEditor/Private/ProductionSettingsCustomization.cpp b/Engine/Plugins/VirtualProduction/Shots/Source/ShotsEditor/Private/ProductionSettingsCustomization.cpp
--- a/Engine/Plugins/VirtualProduction/Shots/Source/ShotsEditor/Private/ProductionSettingsCustomization.cpp
+++ b/Engine/Plugins/VirtualProduction/Shots/Source/ShotsEditor/Private/ProductionSettingsCustomization.cpp
@@ -66,7 +66,7 @@ TSharedRef<SWidget> FProductionSettingsCustomization::BuildProductionNameMenu()
 		FText::FromName(NoActiveProductionName),
 		FText::FromName(NoActiveProductionName),
 		FSlateIcon(),
-		FUIAction(FExecuteAction::CreateRaw(this, &FProductionSettingsCustomization::SetActiveProduction, NoActiveProductionName)),
+		FUIAction(FExecuteAction::CreateRaw(this, &FProductionSettingsCustomization::SetActiveProduction, FGuid())),
 		NAME_None,
 		EUserInterfaceActionType::None
 	);
@@ -74,19 +74,13 @@ TSharedRef<SWidget> FProductionSettingsCustomization::BuildProductionNameMenu()
 	const UProductionSettings* const ProductionSettings = GetDefault<UProductionSettings>();
 
 	// Add a menu option with the production name for each production available in this project
-	TArray<FName> ProductionNames;
 	for (const FCinematicProduction& Production : ProductionSettings->Productions)
-	{
-		ProductionNames.Add(Production.ProductionName);
-	}
-
-	for (FName Name : ProductionNames)
 	{
 		MenuBuilder.AddMenuEntry(
-			FText::FromName(Name),
-			FText::FromName(Name),
+			FText::FromString(Production.ProductionName),
+			FText::FromString(Production.ProductionName),
 			FSlateIcon(),
-			FUIAction(FExecuteAction::CreateRaw(this, &FProductionSettingsCustomization::SetActiveProduction, Name)),
+			FUIAction(FExecuteAction::CreateRaw(this, &FProductionSettingsCustomization::SetActiveProduction, Production.ProductionID)),
 			NAME_None,
 			EUserInterfaceActionType::None
 		);
@@ -98,13 +92,17 @@ TSharedRef<SWidget> FProductionSettingsCustomization::BuildProductionNameMenu()
 FText FProductionSettingsCustomization::GetActiveProductionName() const
 {
  	const UProductionSettings* const ProductionSettings = GetDefault<UProductionSettings>();
-	return FText::FromName(ProductionSettings->ActiveProductionName);
+	if (ProductionSettings->ActiveProductionName.IsEmpty())
+	{
+		return FText::FromName(NAME_None);
+	}
+	return FText::FromString(ProductionSettings->ActiveProductionName);
 }
 
-void FProductionSettingsCustomization::SetActiveProduction(FName ProductionName)
+void FProductionSettingsCustomization::SetActiveProduction(FGuid ProductionID)
 {
 	UProductionSettings* ProductionSettings = GetMutableDefault<UProductionSettings>();
-	ProductionSettings->SetActiveProduction(ProductionName);
+	ProductionSettings->SetActiveProduction(ProductionID);
 
 	// The active production influences whether certain sequencer settings are writable or read-only.
 	// This forces the details view of the other sequencer settings categories to refresh immediately to respect the new property flags.
diff --git a/Engine/Plugins/VirtualProduction/Shots/Source/ShotsEditor/Private/ProductionSettingsCustomization.h b/Engine/Plugins/VirtualProduction/Shots/Source/ShotsEditor/Private/ProductionSettingsCustomization.h
--- a/Engine/Plugins/VirtualProduction/Shots/Source/ShotsEditor/Private/ProductionSettingsCustomization.h
+++ b/Engine/Plugins/VirtualProduction/Shots/Source/ShotsEditor/Private/ProductionSettingsCustomization.h
@@ -24,7 +24,7 @@ public:
 
 private:
 	/** Set the active production of the production project settings */
-	void SetActiveProduction(FName ProductionName);
+	void SetActiveProduction(FGuid ProductionID);
 
 	/** Get the name of the currently active production */
 	FText GetActiveProductionName() const;
diff --git a/Engine/Plugins/VirtualProduction/Shots/Source/ShotsEditor/Private/SProductionWizard.cpp b/Engine/Plugins/VirtualProduction/Shots/Source/ShotsEditor/Private/SProductionWizard.cpp
new file mode 100644
--- /dev/null
+++ b/Engine/Plugins/VirtualProduction/Shots/Source/ShotsEditor/Private/SProductionWizard.cpp
@@ -0,0 +1,978 @@
+// Copyright Epic Games, Inc. All Rights Reserved.
+
+#include "SProductionWizard.h"
+
+#include "Framework/MultiBox/MultiBoxBuilder.h"
+#include "ProductionSettings.h"
+#include "ShotsEditorStyle.h"
+#include "Styling/StyleColors.h"
+#include "Widgets/Input/SComboButton.h"
+#include "Widgets/Input/SNumericEntryBox.h"
+#include "Widgets/Layout/SBorder.h"
+#include "Widgets/Layout/SSeparator.h"
+#include "Widgets/Layout/SSplitter.h"
+#include "Widgets/SBoxPanel.h"
+#include "Widgets/SFrameRatePicker.h"
+
+#define LOCTEXT_NAMESPACE "SProductionWizard"
+
+namespace UE::Shots::ProductionWizard::Private
+{
+	auto IsActiveProductionValid = []() -> bool
+		{
+			const UProductionSettings* const ProductionSettings = GetDefault<UProductionSettings>();
+			return (ProductionSettings->GetActiveProduction() != nullptr);
+		};
+
+	auto GetActiveProductionName = []() -> FText
+		{
+			const UProductionSettings* const ProductionSettings = GetDefault<UProductionSettings>();
+			return FText::FromString(ProductionSettings->ActiveProductionName);
+		};
+
+	auto OnSetActiveProduction = [](FGuid ProductionID)
+		{
+			UProductionSettings* ProductionSettings = GetMutableDefault<UProductionSettings>();
+			if (ProductionSettings->IsActiveProduction(ProductionID))
+			{
+				ProductionSettings->SetActiveProduction(FGuid());
+			}
+			else
+			{
+				ProductionSettings->SetActiveProduction(ProductionID);
+			}
+		};
+};
+
+void SProductionWizard::Construct(const FArguments& InArgs)
+{
+	// Subscribe to be notified when the Production Settings list of productions has changed (for example, if a production was added/removed)
+	UProductionSettings* ProductionSettings = GetMutableDefault<UProductionSettings>();
+	ProductionListChangedHandle = ProductionSettings->OnProductionListChanged().AddSP(this, &SProductionWizard::UpdateProductionList);
+
+	ChildSlot
+		[
+			SNew(SBorder)
+				.Padding(0.f)
+				.BorderImage(FShotsEditorStyle::Get().GetBrush("ProductionWizard.PanelBackground"))
+				[
+					SNew(SVerticalBox)
+
+					+ SVerticalBox::Slot()
+						[
+							SNew(SSplitter)
+								.Orientation(Orient_Horizontal)
+								.PhysicalSplitterHandleSize(2.0f)
+
+							+ SSplitter::Slot()
+								.Value(0.25f)
+								.MinSize(192.0f)
+								.Expose(MenuPanelSlot)
+
+							+ SSplitter::Slot()
+								.Value(0.75f)
+								.Expose(ContentPanelSlot)
+						]
+
+					+ SVerticalBox::Slot()
+						.AutoHeight()
+						[
+							SNew(SSeparator)
+								.Orientation(Orient_Horizontal)
+								.Thickness(2.0f)
+						]
+
+					+ SVerticalBox::Slot()
+						.AutoHeight()
+						[
+							MakeButtonsPanel()
+						]
+				]
+		];
+
+	MenuPanelSlot->AttachWidget(MakeMenuPanel());
+}
+
+SProductionWizard::~SProductionWizard()
+{
+	UProductionSettings* ProductionSettings = GetMutableDefault<UProductionSettings>();
+	ProductionSettings->OnProductionListChanged().Remove(ProductionListChangedHandle);
+}
+
+void SProductionWizard::UpdateProductionList()
+{
+	ProductionList.Empty();
+
+	const UProductionSettings* const ProductionSettings = GetDefault<UProductionSettings>();
+	for (const FCinematicProduction& Production : ProductionSettings->Productions)
+	{
+		ProductionList.Add(MakeShared<FCinematicProduction>(Production));
+	}
+
+	if (ProductionListView)
+	{
+		ProductionListView->RequestListRefresh();
+	}
+}
+
+TSharedRef<SWidget> SProductionWizard::MakeButtonsPanel()
+{
+	return SNew(SBorder)
+		.BorderImage(FAppStyle::Get().GetBrush("Brushes.Panel"))
+		.Padding(16.0f)
+		.HAlign(HAlign_Right)
+		[
+			SNew(SHorizontalBox)
+
+			+ SHorizontalBox::Slot()
+				.MinWidth(118.0f)
+				.MaxWidth(118.0f)
+				.VAlign(VAlign_Center)
+				.Padding(0.0f, 0.0f, 8.0f, 0.0f)
+				[
+					SNew(SButton)
+						.Text(LOCTEXT("BackButton", "Back"))
+						.HAlign(HAlign_Center)
+				]
+
+			+ SHorizontalBox::Slot()
+				.MinWidth(118.0f)
+				.MaxWidth(118.0f)
+				.VAlign(VAlign_Center)
+				[
+					SNew(SButton)
+						.Text(LOCTEXT("NextButton", "Next"))
+						.HAlign(HAlign_Center)
+				]
+		];
+}
+
+TSharedRef<SWidget> SProductionWizard::MakeMenuPanel()
+{
+	// Build the Revision Control menu entry, including its content panel
+	TSharedPtr<FProductionWizardMenuEntry> RevisionControlMenuEntry = MakeShared<FProductionWizardMenuEntry>();
+	RevisionControlMenuEntry->Label = LOCTEXT("RevisionControlMenuLabel", "Revision Control");
+	RevisionControlMenuEntry->Icon = FSlateIcon(FShotsEditorStyle::StyleName, "Icons.RevisionControl");
+	RevisionControlMenuEntry->Panel = MakeRevisionControlPanel();
+
+	// Build the Productions menu entry, including its content panel
+	TSharedPtr<FProductionWizardMenuEntry> ProductionMenuEntry = MakeShared<FProductionWizardMenuEntry>();
+	ProductionMenuEntry->Label = LOCTEXT("ProductionsMenuLabel", "Productions");
+	ProductionMenuEntry->Icon = FSlateIcon(FShotsEditorStyle::StyleName, "Icons.Productions");
+	ProductionMenuEntry->Panel = MakeProductionsPanel();
+
+	UserMenuEntries.Add(MoveTemp(RevisionControlMenuEntry));
+	UserMenuEntries.Add(MoveTemp(ProductionMenuEntry));
+
+	// Build the Sequencer Settings menu entry, including its content panel
+	TSharedPtr<FProductionWizardMenuEntry> SequencerSettingsMenuEntry = MakeShared<FProductionWizardMenuEntry>();
+	SequencerSettingsMenuEntry->Label = LOCTEXT("SequencerSettingsMenuLabel", "Sequencer Settings");
+	SequencerSettingsMenuEntry->Icon = FSlateIcon(FShotsEditorStyle::StyleName, "Icons.Sequencer");
+	SequencerSettingsMenuEntry->Panel = MakeSequencerSettingsPanel();
+
+	// Build the Naming Tokens menu entry, including its content panel
+	TSharedPtr<FProductionWizardMenuEntry> NamingTokensMenuEntry = MakeShared<FProductionWizardMenuEntry>();
+	NamingTokensMenuEntry->Label = LOCTEXT("NamingTokensMenuLabel", "Naming Tokens");
+	NamingTokensMenuEntry->Icon = FSlateIcon(FShotsEditorStyle::StyleName, "Icons.NamingTokens");
+	NamingTokensMenuEntry->Panel = MakeNamingTokensPanel();
+
+	// Build the Asset Naming menu entry, including its content panel
+	TSharedPtr<FProductionWizardMenuEntry> AssetNamingMenuEntry = MakeShared<FProductionWizardMenuEntry>();
+	AssetNamingMenuEntry->Label = LOCTEXT("AssetNamingMenuLabel", "Asset Naming");
+	AssetNamingMenuEntry->Icon = FSlateIcon(FShotsEditorStyle::StyleName, "Icons.AssetNaming");
+	AssetNamingMenuEntry->Panel = MakeAssetNamingPanel();
+
+	// Build the Folder Hierarchy menu entry, including its content panel
+	TSharedPtr<FProductionWizardMenuEntry> FolderHierarchyMenuEntry = MakeShared<FProductionWizardMenuEntry>();
+	FolderHierarchyMenuEntry->Label = LOCTEXT("FolderHierarchyMenuLabel", "Folder Hierarchy");
+	FolderHierarchyMenuEntry->Icon = FSlateIcon(FShotsEditorStyle::StyleName, "Icons.Folder");
+	FolderHierarchyMenuEntry->Panel = MakeFolderHierarchyPanel();
+
+	ProductionMenuEntries.Add(SequencerSettingsMenuEntry);
+	ProductionMenuEntries.Add(NamingTokensMenuEntry);
+	ProductionMenuEntries.Add(AssetNamingMenuEntry);
+	ProductionMenuEntries.Add(FolderHierarchyMenuEntry);
+
+	// Construct the list view for the user menu entries
+	UserMenuListView = SNew(SListView<TSharedPtr<FProductionWizardMenuEntry>>)
+		.ListItemsSource(&UserMenuEntries)
+		.OnGenerateRow(this, &SProductionWizard::OnGenerateMenuRow)
+		.OnSelectionChanged(this, &SProductionWizard::OnUserMenuSelectionChanged);
+
+	// Construct the list view for the production menu entries
+	ProductionMenuListView = SNew(SListView<TSharedPtr<FProductionWizardMenuEntry>>)
+		.ListItemsSource(&ProductionMenuEntries)
+		.OnGenerateRow(this, &SProductionWizard::OnGenerateMenuRow)
+		.OnSelectionChanged(this, &SProductionWizard::OnProductionMenuSelectionChanged);
+
+	// The initial selection upon opening the wizard is the first user menu entry (revision control)
+	UserMenuListView->SetSelection(UserMenuEntries[0]);
+	ContentPanelSlot->AttachWidget(UserMenuEntries[0]->Panel.ToSharedRef());
+
+	return SNew(SBorder)
+		.BorderImage(FShotsEditorStyle::Get().GetBrush("ProductionWizard.RecessedBackground"))
+		.Padding(16.0f)
+		[
+			SNew(SVerticalBox)
+
+			+ SVerticalBox::Slot()
+				.AutoHeight()
+				.Padding(0.0f, 0.0f, 0.0f, 8.0f)
+				[
+					SNew(STextBlock)
+						.Text(LOCTEXT("UserSetupMenuGroup", "User Setup"))
+						.Font(FAppStyle::Get().GetFontStyle("NormalFontBold"))
+				]
+
+			+ SVerticalBox::Slot()
+				.AutoHeight()
+				.Padding(0.0f, 0.0f, 0.0f, 8.0f)
+				[
+					UserMenuListView.ToSharedRef()
+				]
+
+			+ SVerticalBox::Slot()
+				.AutoHeight()
+				.Padding(0.0f, 0.0f, 0.0f, 8.0f)
+				[
+					SNew(STextBlock)
+						.Text(LOCTEXT("ProductionSettingsMenuGroup", "Production Settings"))
+						.Font(FAppStyle::Get().GetFontStyle("NormalFontBold"))
+				]
+
+			+ SVerticalBox::Slot()
+				.AutoHeight()
+				[
+					ProductionMenuListView.ToSharedRef()
+				]
+		];
+}
+
+TSharedRef<ITableRow> SProductionWizard::OnGenerateMenuRow(TSharedPtr<FProductionWizardMenuEntry> MenuEntry, const TSharedRef<STableViewBase>& OwnerTable)
+{
+	return SNew(STableRow<TSharedPtr<FProductionWizardMenuEntry>>, OwnerTable)
+		.ShowSelection(true)
+		.Padding(FMargin(4.0f, 4.0f, 0.0f, 4.0f))
+		.Content()
+		[
+			SNew(SHorizontalBox)
+
+			+ SHorizontalBox::Slot()
+				.AutoWidth()
+				.Padding(0.0f, 0.0f, 4.0f, 0.0f)
+				[
+					SNew(SImage).Image(MenuEntry->Icon.GetIcon())
+				]
+
+				+ SHorizontalBox::Slot()
+				.AutoWidth()
+				[
+					SNew(STextBlock).Text(MenuEntry->Label)
+				]
+		];
+}
+
+void SProductionWizard::OnUserMenuSelectionChanged(TSharedPtr<FProductionWizardMenuEntry> SelectedEntry, ESelectInfo::Type SelectInfo)
+{
+	if (SelectInfo != ESelectInfo::Direct)
+	{
+		check((SelectedEntry->Panel));
+
+		ProductionMenuListView->ClearSelection();
+		ContentPanelSlot->AttachWidget(SelectedEntry->Panel.ToSharedRef());
+	}
+}
+
+void SProductionWizard::OnProductionMenuSelectionChanged(TSharedPtr<FProductionWizardMenuEntry> SelectedEntry, ESelectInfo::Type SelectInfo)
+{
+	if (SelectInfo != ESelectInfo::Direct)
+	{
+		check((SelectedEntry->Panel));
+
+		UserMenuListView->ClearSelection();
+		ContentPanelSlot->AttachWidget(SelectedEntry->Panel.ToSharedRef());
+	}
+}
+
+TSharedRef<SWidget> SProductionWizard::MakeRevisionControlPanel()
+{
+	// TODO: Fill in this stubbed out panel
+	return SNew(SVerticalBox)
+
+	+ SVerticalBox::Slot()
+		.FillHeight(1.0f)
+		.HAlign(HAlign_Center)
+		[
+			SNew(SHorizontalBox)
+
+			+ SHorizontalBox::Slot()
+				.FillWidth(1.0f)
+				.VAlign(VAlign_Center)
+				[
+					SNew(STextBlock).Text(LOCTEXT("RevisionControlPanel", "REVISION CONTROL PANEL"))
+				]
+		];
+}
+
+TSharedRef<SWidget> SProductionWizard::MakeProductionsPanel()
+{
+	// Build the list of productions for the list view before building the widget for the first time
+	UpdateProductionList();
+
+	auto OnCreateNewProductionClicked = []() -> FReply
+		{
+			UProductionSettings* ProductionSettings = GetMutableDefault<UProductionSettings>();
+			ProductionSettings->AddProduction();
+			return FReply::Handled();
+		};
+
+	return SNew(SBorder)
+		.BorderImage(FAppStyle::Get().GetBrush("Brushes.Panel"))
+		.Padding(16.0f)
+		[
+			SNew(SVerticalBox)
+
+			// Title
+			+ SVerticalBox::Slot()
+				.AutoHeight()
+				.Padding(0.0f, 0.0f, 0.0f, 8.0f)
+				[
+					SNew(STextBlock)
+						.Text(LOCTEXT("ProductionsTitle", "User Setup"))
+						.Font(FShotsEditorStyle::Get().GetFontStyle("ProductionWizard.TitleFont"))
+				]
+
+			// Heading
+			+ SVerticalBox::Slot()
+				.AutoHeight()
+				.Padding(0.0f, 0.0f, 0.0f, 8.0f)
+				[
+					SNew(STextBlock)
+						.Text(LOCTEXT("ProductionsHeading", "Productions"))
+						.Font(FShotsEditorStyle::Get().GetFontStyle("ProductionWizard.HeadingFont"))
+				]
+
+			// Info Text
+			+ SVerticalBox::Slot()
+				.AutoHeight()
+				.Padding(0.0f, 0.0f, 0.0f, 20.0f)
+				[
+					SNew(STextBlock)
+						.Text(LOCTEXT("ProductionsInfoText", "You can use this Unreal project for multiple shows, and assign different settings for each show.\n"
+							"Create a production here for each show, then choose one to be the active production."))
+				]
+
+			// Create / Import Buttons
+			+ SVerticalBox::Slot()
+				.AutoHeight()
+				.Padding(0.0f, 0.0f, 0.0f, 4.0f)
+				[
+					SNew(SHorizontalBox)
+
+					+ SHorizontalBox::Slot()
+						.AutoWidth()
+						[
+							SNew(SButton)
+								.ContentPadding(FMargin(2.0f))
+								.OnClicked_Lambda(OnCreateNewProductionClicked)
+								[
+									SNew(SHorizontalBox)
+
+									+ SHorizontalBox::Slot()
+										.AutoWidth()
+										.Padding(0.0f, 0.0f, 4.0f, 0.0f)
+										[
+											SNew(SImage)
+												.Image(FAppStyle::Get().GetBrush("Icons.Plus"))
+												.ColorAndOpacity(FStyleColors::AccentGreen)
+										]
+
+
+									+ SHorizontalBox::Slot()
+										.AutoWidth()
+										[
+											SNew(STextBlock).Text(LOCTEXT("CreateNewProductionButton", "Create a New Production"))
+										]
+								]
+						]
+
+					+ SHorizontalBox::Slot()
+						.AutoWidth()
+						[
+							// TODO: Add support for importing a production when this button is clicked
+							SNew(SButton)
+								.ContentPadding(FMargin(2.0f))
+								[
+									SNew(SHorizontalBox)
+
+									+ SHorizontalBox::Slot()
+										.AutoWidth()
+										.Padding(0.0f, 0.0f, 4.0f, 0.0f)
+										[
+											SNew(SImage)
+												.Image(FAppStyle::Get().GetBrush("Icons.Import"))
+										]
+
+									+ SHorizontalBox::Slot()
+										.AutoWidth()
+										[
+											SNew(STextBlock)
+												.Text(LOCTEXT("ImportProductionButton", "Import Production"))
+												.ToolTipText(LOCTEXT("ImportProductionButtonToolTip", "Not yet implemented"))
+										]
+								]
+						]
+				]
+
+			// Production List View
+			+ SVerticalBox::Slot()
+				.FillHeight(1.0f)
+				[
+					SNew(SBorder)
+						.BorderImage(FAppStyle::Get().GetBrush("Brushes.Recessed"))
+						.Padding(8.0f)
+						[
+							SAssignNew(ProductionListView, SListView<TSharedPtr<FCinematicProduction>>)
+								.ListItemsSource(&ProductionList)
+								.SelectionMode(ESelectionMode::None)
+								.OnGenerateRow(this, &SProductionWizard::OnGenerateProductionRow)
+								.OnItemsRebuilt(this, &SProductionWizard::OnProductionListItemsRebuilt)
+						]
+				]
+		];
+}
+
+void SProductionWizard::OnProductionListItemsRebuilt()
+{
+	// If there was a production with no name, then ProductionListButtonToRename will be valid, and it can be put into edit mode so the user can rename the production.
+	if (ProductionListButtonToRename)
+	{
+		ProductionListButtonToRename->EnterEditMode();
+		ProductionListButtonToRename = nullptr;
+	}
+}
+
+TSharedRef<SWidget> SProductionWizard::MakeSequencerSettingsPanel()
+{
+	using namespace UE::Shots::ProductionWizard::Private;
+
+	UProductionSettings* ProductionSettings = GetMutableDefault<UProductionSettings>();
+
+	// Get a text representation of the active subsequence priority to display in the combo button below
+	auto GetSubsequencePriorityText = [ProductionSettings]() -> FText
+		{
+			if (ProductionSettings->GetActiveSubsequencePriority() == ESubsequencePriority::TopDown)
+			{
+				return LOCTEXT("TopDownText", "Top Down");
+			}
+			return LOCTEXT("BottomUpText", "Bottom Up");
+		};
+
+	// Build the menu options for the subsequence priority combo button below
+	FMenuBuilder SubsequencePriorityMenuBuilder(true, nullptr);
+	{
+		SubsequencePriorityMenuBuilder.AddMenuEntry(
+			LOCTEXT("TopDownText", "Top Down"),
+			LOCTEXT("TopDownText", "Top Down"),
+			FSlateIcon(),
+			FUIAction(FExecuteAction::CreateLambda([ProductionSettings]() { ProductionSettings->SetActiveSubsequencePriority(ESubsequencePriority::TopDown); })),
+			NAME_None,
+			EUserInterfaceActionType::Button
+		);
+
+		SubsequencePriorityMenuBuilder.AddMenuEntry(
+			LOCTEXT("BottomUpText", "Bottom Up"),
+			LOCTEXT("BottomUpText", "Bottom Up"),
+			FSlateIcon(),
+			FUIAction(FExecuteAction::CreateLambda([ProductionSettings]() { ProductionSettings->SetActiveSubsequencePriority(ESubsequencePriority::BottomUp); })),
+			NAME_None,
+			EUserInterfaceActionType::Button
+		);
+	}
+
+	return SNew(SVerticalBox)
+
+	// Active Production Selector
+	+ SVerticalBox::Slot()
+		.AutoHeight()
+		[
+			MakeActiveProductionWidget()
+		]
+
+	// Separator
+	+ SVerticalBox::Slot()
+		.AutoHeight()
+		[
+			SNew(SSeparator)
+				.Orientation(Orient_Horizontal)
+				.Thickness(2.0f)
+		]
+
+	// Sequencer Settings Panel
+	+ SVerticalBox::Slot()
+		.FillHeight(1.0f)
+		[
+			SNew(SBorder)
+				.BorderImage(FAppStyle::Get().GetBrush("Brushes.Panel"))
+				.Padding(16.0f)
+				[
+					SNew(SVerticalBox)
+
+					// Title
+					+ SVerticalBox::Slot()
+						.AutoHeight()
+						.Padding(0.0f, 0.0f, 0.0f, 8.0f)
+						[
+							SNew(STextBlock)
+								.Text(LOCTEXT("SequencerSettingsTitle", "Production Settings"))
+								.Font(FShotsEditorStyle::Get().GetFontStyle("ProductionWizard.TitleFont"))
+						]
+
+					// Heading
+					+ SVerticalBox::Slot()
+						.AutoHeight()
+						.Padding(0.0f, 0.0f, 0.0f, 8.0f)
+						[
+							SNew(STextBlock)
+								.Text(LOCTEXT("SequencerSettingsHeading", "Sequencer Settings"))
+								.Font(FShotsEditorStyle::Get().GetFontStyle("ProductionWizard.HeadingFont"))
+						]
+
+					// Info Text
+					+ SVerticalBox::Slot()
+						.AutoHeight()
+						.Padding(0.0f, 0.0f, 0.0f, 20.0f)
+						[
+							SNew(STextBlock)
+								.Text(LOCTEXT("SequencerSettingsInfoText", "Configure basic settings that apply to all level sequences and shots."))
+						]
+
+					// Frame Rate setting
+					+ SVerticalBox::Slot()
+						.AutoHeight()
+						.Padding(0.0f, 0.0f, 0.0f, 8.0f)
+						[
+							SNew(SHorizontalBox)
+
+							+ SHorizontalBox::Slot()
+								.FillWidth(0.3f)
+								.VAlign(VAlign_Center)
+								[
+									SNew(STextBlock).Text(LOCTEXT("FrameRateText", "Frame Rate"))
+								]
+
+							+ SHorizontalBox::Slot()
+								.FillWidth(0.7f)
+								.VAlign(VAlign_Center)
+								[
+									SNew(SFrameRatePicker)
+										.HasMultipleValues(false)
+										.Value_Lambda([ProductionSettings]() { return ProductionSettings->GetActiveDisplayRate(); })
+										.OnValueChanged_Lambda([ProductionSettings](FFrameRate NewFrameRate) { ProductionSettings->SetActiveDisplayRate(NewFrameRate); })
+										.IsEnabled_Lambda(IsActiveProductionValid)
+								]
+						]
+
+					// Start Frame setting
+					+ SVerticalBox::Slot()
+						.AutoHeight()
+						.Padding(0.0f, 0.0f, 0.0f, 8.0f)
+						[
+							SNew(SHorizontalBox)
+
+							+ SHorizontalBox::Slot()
+								.FillWidth(0.3f)
+								.VAlign(VAlign_Center)
+								[
+									SNew(STextBlock).Text(LOCTEXT("SequencerStartFrameText", "Sequencer Start Frame"))
+								]
+
+							+ SHorizontalBox::Slot()
+								.FillWidth(0.7f)
+								.VAlign(VAlign_Center)
+								[
+									SNew(SNumericEntryBox<int32>)
+										.AllowSpin(false)
+										.Value_Lambda([ProductionSettings]() { return ProductionSettings->GetActiveStartFrame(); })
+										.OnValueCommitted_Lambda([ProductionSettings](int32 NewValue, ETextCommit::Type CommitType) { ProductionSettings->SetActiveStartFrame(NewValue); })
+										.IsEnabled_Lambda(IsActiveProductionValid)
+								]
+						]
+
+					// Subsequence Priority setting
+					+ SVerticalBox::Slot()
+						.AutoHeight()
+						.Padding(0.0f, 0.0f, 0.0f, 8.0f)
+						[
+							SNew(SHorizontalBox)
+
+							+ SHorizontalBox::Slot()
+								.FillWidth(0.3f)
+								.VAlign(VAlign_Center)
+								[
+									SNew(STextBlock).Text(LOCTEXT("SubsequencePriorityText", "Subsequence Priority"))
+								]
+
+							+ SHorizontalBox::Slot()
+								.FillWidth(0.7f)
+								.VAlign(VAlign_Center)
+								[
+									SNew(SComboButton)
+										.IsEnabled_Lambda(IsActiveProductionValid)
+										.MenuContent()
+										[
+											SubsequencePriorityMenuBuilder.MakeWidget()
+										]
+										.ButtonContent()
+										[
+											SNew(STextBlock)
+												.Text_Lambda(GetSubsequencePriorityText)
+												.Font(FAppStyle::Get().GetFontStyle("Normal"))
+										]
+								]
+						]
+				]
+		];
+}
+
+TSharedRef<SWidget> SProductionWizard::MakeNamingTokensPanel()
+{
+	// TODO: Fill in this stubbed out panel
+	return SNew(SVerticalBox)
+
+	+ SVerticalBox::Slot()
+		.FillHeight(1.0f)
+		.HAlign(HAlign_Center)
+		[
+			SNew(SHorizontalBox)
+
+			+ SHorizontalBox::Slot()
+				.FillWidth(1.0f)
+				.VAlign(VAlign_Center)
+				[
+					SNew(STextBlock).Text(LOCTEXT("NamingTokensPanel", "NAMING TOKENS PANEL"))
+				]
+		];
+}
+
+TSharedRef<SWidget> SProductionWizard::MakeAssetNamingPanel()
+{
+	// TODO: Fill in this stubbed out panel
+	return SNew(SVerticalBox)
+
+	+ SVerticalBox::Slot()
+		.FillHeight(1.0f)
+		.HAlign(HAlign_Center)
+		[
+			SNew(SHorizontalBox)
+
+			+ SHorizontalBox::Slot()
+				.FillWidth(1.0f)
+				.VAlign(VAlign_Center)
+				[
+					SNew(STextBlock).Text(LOCTEXT("AssetNamingPanel", "ASSET NAMING PANEL"))
+				]
+		];
+}
+
+TSharedRef<SWidget> SProductionWizard::MakeFolderHierarchyPanel()
+{
+	// TODO: Fill in this stubbed out panel
+	return SNew(SVerticalBox)
+
+	+ SVerticalBox::Slot()
+		.FillHeight(1.0f)
+		.HAlign(HAlign_Center)
+		[
+			SNew(SHorizontalBox)
+
+			+ SHorizontalBox::Slot()
+				.FillWidth(1.0f)
+				.VAlign(VAlign_Center)
+				[
+					SNew(STextBlock).Text(LOCTEXT("FolderHierarchyPanel", "FOLDER HIERARCHY PANEL"))
+				]
+		];
+}
+
+TSharedRef<SWidget> SProductionWizard::MakeActiveProductionWidget()
+{
+	using namespace UE::Shots::ProductionWizard::Private;
+
+	auto BuildMenu = []() -> TSharedRef<SWidget>
+		{
+			FMenuBuilder MenuBuilder(true, nullptr);
+
+			// None option, allowing the user to set no active production
+			MenuBuilder.AddMenuEntry(
+				FText::FromName(NAME_None),
+				FText::FromName(NAME_None),
+				FSlateIcon(),
+				FUIAction(FExecuteAction::CreateLambda(OnSetActiveProduction, FGuid())),
+				NAME_None,
+				EUserInterfaceActionType::Button
+			);
+
+			const UProductionSettings* const ProductionSettings = GetDefault<UProductionSettings>();
+			for (const FCinematicProduction& Production : ProductionSettings->Productions)
+			{
+				MenuBuilder.AddMenuEntry(
+					FText::FromString(Production.ProductionName),
+					FText::FromString(Production.ProductionName),
+					FSlateIcon(),
+					FUIAction(FExecuteAction::CreateLambda(OnSetActiveProduction, Production.ProductionID)),
+					NAME_None,
+					EUserInterfaceActionType::Button
+				);
+			}
+
+			return MenuBuilder.MakeWidget();
+		};
+
+	return SNew(SBorder)
+		.BorderImage(FAppStyle::Get().GetBrush("Brushes.Panel"))
+		.Padding(16.0f, 8.0f)
+		[
+			SNew(SHorizontalBox)
+
+			+ SHorizontalBox::Slot()
+				.AutoWidth()
+				.VAlign(VAlign_Center)
+				.Padding(0.0f, 0.0f, 8.0f, 0.0f)
+				[
+					SNew(STextBlock).Text(LOCTEXT("ActiveProductionText", "Active Production"))
+				]
+
+			+ SHorizontalBox::Slot()
+				.FillWidth(1.0f)
+				.VAlign(VAlign_Center)
+				[
+					SNew(SComboButton)
+						.OnGetMenuContent_Lambda(BuildMenu)
+						.ButtonContent()
+						[
+							SNew(STextBlock).Text_Lambda(GetActiveProductionName)
+						]
+				]
+		];
+}
+
+TSharedRef<ITableRow> SProductionWizard::OnGenerateProductionRow(TSharedPtr<FCinematicProduction> InItem, const TSharedRef<STableViewBase>& OwnerTable)
+{
+	auto OnDeleteProduction = [](FGuid ProductionID) -> FReply
+		{
+			UProductionSettings* ProductionSettings = GetMutableDefault<UProductionSettings>();
+			ProductionSettings->DeleteProduction(ProductionID);
+			return FReply::Handled();
+		};
+
+	TSharedRef<SProductionListButton> ListButton = SNew(SProductionListButton)
+		.ProductionID(InItem->ProductionID);
+
+	// If this production was newly created, it will not have a name yet, so the user should be prompted to give it one.
+	// However, the production's list button cannot be put into edit mode and focused until after the widget is created.
+	// Therefore, the widget is saved, and will be put into edit mode after the items for this ListView are finished being rebuilt. 
+	if (InItem->ProductionName.IsEmpty())
+	{
+		ProductionListButtonToRename = ListButton;
+	}
+
+	return SNew(STableRow<TSharedPtr<FName>>, OwnerTable)
+		.ShowSelection(true)
+		.Content()
+		[
+			SNew(SHorizontalBox)
+
+			+ SHorizontalBox::Slot()
+				.FillWidth(1.0f)
+				[
+					ListButton
+				]
+
+			+ SHorizontalBox::Slot()
+				.AutoWidth()
+				.VAlign(VAlign_Center)
+				[
+					// TODO: Add support for exporting a production when this button is clicked
+					SNew(SButton)
+						.ButtonStyle(FShotsEditorStyle::Get(), "ProductionWizard.RecessedButton")
+						[
+							SNew(SImage).Image(FShotsEditorStyle::Get().GetBrush("Icons.Export"))
+						]
+				]
+
+			+ SHorizontalBox::Slot()
+				.AutoWidth()
+				.VAlign(VAlign_Center)
+				[
+					SNew(SButton)
+						.ButtonStyle(FShotsEditorStyle::Get(), "ProductionWizard.RecessedButton")
+						.OnClicked_Lambda(OnDeleteProduction, InItem->ProductionID)
+						[
+							SNew(SImage).Image(FAppStyle::Get().GetBrush("Icons.Delete"))
+						]
+				]
+		];
+}
+
+SProductionListButton::SProductionListButton()
+	: ProductionID(*this)
+{
+}
+
+void SProductionListButton::Construct(const FArguments& InArgs)
+{
+	using namespace UE::Shots::ProductionWizard::Private;
+
+	ProductionID.Assign(*this, InArgs._ProductionID);
+
+	SButton::Construct(SButton::FArguments());
+	SetButtonStyle(&FShotsEditorStyle::Get().GetWidgetStyle<FButtonStyle>("ProductionWizard.RecessedButton"));
+
+	auto OnProductionButtonClicked = [](FGuid ProductionID) -> FReply
+	{
+		OnSetActiveProduction(ProductionID);
+		return FReply::Handled();
+	};
+
+	SetOnClicked(FOnClicked::CreateLambda(OnProductionButtonClicked, ProductionID.Get()));
+
+	// Returns the color of the checkmark icon based on the active production and button hover state
+	auto GetCheckMarkColor = [this]() -> FSlateColor
+		{
+			const UProductionSettings* const ProductionSettings = GetDefault<UProductionSettings>();
+			if (ProductionSettings->IsActiveProduction(ProductionID.Get()))
+			{
+				return FStyleColors::AccentGreen;
+			}
+			else if (bIsHovered)
+			{
+				return FStyleColors::Foreground;
+			}
+			return FStyleColors::Transparent;
+		};
+
+	// Returns the production name based on the ProductionID assigned to this button
+	auto GetProductionName = [this]() -> FText
+		{
+			const UProductionSettings* const ProductionSettings = GetDefault<UProductionSettings>();
+			if (const FCinematicProduction* Production = ProductionSettings->GetProduction(ProductionID.Get()))
+			{
+				if (Production->ProductionName.IsEmpty())
+				{
+					return FText::FromString(ProductionSettings->GetUniqueProductionName());
+				}
+				return FText::FromString(Production->ProductionName);
+			}
+			return FText::GetEmpty();
+		};
+
+	// Renames this button's production
+	auto OnRenameFinished = [this](const FText& InText, ETextCommit::Type InCommitType)
+		{
+			UProductionSettings* ProductionSettings = GetMutableDefault<UProductionSettings>();
+			ProductionSettings->RenameProduction(ProductionID.Get(), InText.ToString());
+		};
+
+	auto OnVerifyRename = [this](const FText& InText, FText& OutErrorMessage) -> bool
+		{
+			const UProductionSettings* const ProductionSettings = GetDefault<UProductionSettings>();
+
+			if (InText.IsEmpty())
+			{
+				OutErrorMessage = LOCTEXT("EmptyNameErrorMessage", "Please provide a name for the production");
+				return false;
+			}
+
+			bool bIsNameValid = true;
+			for (const FCinematicProduction& Production : ProductionSettings->Productions)
+			{
+				if (Production.ProductionID == ProductionID.Get())
+				{
+					continue;
+				}
+
+				if (InText.ToString() == Production.ProductionName)
+				{
+					bIsNameValid = false;
+					OutErrorMessage = LOCTEXT("ExistingNameErrorMessage", "A production already exists with this name");
+				}
+			}
+			return bIsNameValid;
+		};
+
+	ChildSlot
+		[
+			SNew(SHorizontalBox)
+
+			+ SHorizontalBox::Slot()
+				.AutoWidth()
+				.VAlign(VAlign_Center)
+				.Padding(0.0f, 0.0f, 8.0f, 0.0f)
+				[
+					SNew(SImage)
+						.Image(FAppStyle::Get().GetBrush("Icons.Check"))
+						.ColorAndOpacity_Lambda(GetCheckMarkColor)
+				]
+
+			+ SHorizontalBox::Slot()
+				.FillWidth(1.0f)
+				[
+					SAssignNew(EditableTextBlock, SInlineEditableTextBlock)
+						.Text_Lambda(GetProductionName)
+						.OnVerifyTextChanged_Lambda(OnVerifyRename)
+						.OnTextCommitted_Lambda(OnRenameFinished)
+						.IsSelected_Lambda([]() { return false; }) // Disable double-select to rename
+				]
+		];
+}
+
+void SProductionListButton::EnterEditMode()
+{
+	EditableTextBlock->EnterEditingMode();
+}
+
+FReply SProductionListButton::OnMouseButtonDown(const FGeometry& MyGeometry, const FPointerEvent& MouseEvent)
+{
+	// Create the context menu to be launched on right mouse click
+	if (MouseEvent.GetEffectingButton() == EKeys::RightMouseButton)
+	{
+		FMenuBuilder MenuBuilder(true, nullptr);
+		MenuBuilder.AddMenuEntry(
+			LOCTEXT("RenameProduction", "Rename"),
+			LOCTEXT("RenameProductionToolTip", "Rename production"),
+			FSlateIcon(FShotsEditorStyle::StyleName, "Icons.AssetNaming"),
+			FUIAction(FExecuteAction::CreateSP(this, &SProductionListButton::EnterEditMode)),
+			NAME_None,
+			EUserInterfaceActionType::Button
+		);
+
+		FWidgetPath WidgetPath = MouseEvent.GetEventPath() != nullptr ? *MouseEvent.GetEventPath() : FWidgetPath();
+
+		FSlateApplication::Get().PushMenu(
+			AsShared(),
+			WidgetPath,
+			MenuBuilder.MakeWidget(),
+			MouseEvent.GetScreenSpacePosition(),
+			FPopupTransitionEffect(FPopupTransitionEffect::ContextMenu)
+		);
+
+		return FReply::Handled();
+	}
+
+	return SButton::OnMouseButtonDown(MyGeometry, MouseEvent);
+}
+
+void SProductionListButton::OnMouseEnter(const FGeometry& MyGeometry, const FPointerEvent& MouseEvent)
+{
+	SButton::OnMouseEnter(MyGeometry, MouseEvent);
+	bIsHovered = true;
+}
+
+void SProductionListButton::OnMouseLeave(const FPointerEvent& MouseEvent)
+{
+	SButton::OnMouseLeave(MouseEvent);
+	bIsHovered = false;
+}
+
+#undef LOCTEXT_NAMESPACE
diff --git a/Engine/Plugins/VirtualProduction/Shots/Source/ShotsEditor/Private/SProductionWizard.h b/Engine/Plugins/VirtualProduction/Shots/Source/ShotsEditor/Private/SProductionWizard.h
new file mode 100644
--- /dev/null
+++ b/Engine/Plugins/VirtualProduction/Shots/Source/ShotsEditor/Private/SProductionWizard.h
@@ -0,0 +1,142 @@
+// Copyright Epic Games, Inc. All Rights Reserved.
+
+#pragma once
+
+#include "Widgets/SCompoundWidget.h"
+
+#include "Textures/SlateIcon.h"
+#include "Widgets/DeclarativeSyntaxSupport.h"
+#include "Widgets/Layout/SSplitter.h"
+#include "Widgets/Input/SButton.h"
+#include "Widgets/Text/SInlineEditableTextBlock.h"
+#include "Widgets/Views/SListView.h"
+
+struct FCinematicProduction;
+
+/** Structure for the list view of menu entires */
+struct FProductionWizardMenuEntry
+{
+	/** Displayed menu entry text */
+	FText Label;
+
+	/** Displayed menu entry icon */
+	FSlateIcon Icon;
+
+	/** Widget to attach to the content panel slot when this menu entry is selected */
+	TSharedPtr<SWidget> Panel;
+};
+
+/** 
+ * Button used in the Productions panel, containing functionality for a single FCinematicProduction.
+ * This button supports right-click (for renaming) and will show an editable text box for the user to rename a button.
+ */
+class SProductionListButton : public SButton
+{
+public:
+	SProductionListButton();
+
+	SLATE_BEGIN_ARGS(SProductionListButton)
+		: _ProductionID()
+		{ }
+
+		/** ID of the production that this button controls */
+		SLATE_ATTRIBUTE(FGuid, ProductionID)
+
+	SLATE_END_ARGS()
+
+	void Construct(const FArguments& InArgs);
+
+	/** OnMouseButtonDown is overridden because the base class SButton ignores right-clicks */
+	virtual FReply OnMouseButtonDown(const FGeometry& MyGeometry, const FPointerEvent& MouseEvent) override;
+	virtual void OnMouseEnter(const FGeometry& MyGeometry, const FPointerEvent& MouseEvent) override;
+	virtual void OnMouseLeave(const FPointerEvent& MouseEvent) override;
+
+	/** Put the editable text block into edit mode */
+	void EnterEditMode();
+
+private:
+	/** ID of the production that this button controls */
+	TSlateAttribute<FGuid, EInvalidateWidgetReason::Paint> ProductionID;
+
+	/** User-editable text widget for renaming the production */
+	TSharedPtr<SInlineEditableTextBlock> EditableTextBlock;
+
+	/** Hover state, used to determine color/opacity of the button icon */
+	bool bIsHovered = false;
+};
+
+/**
+ * Main UI widget for the Production Wizard Tool
+ */
+class SProductionWizard : public SCompoundWidget
+{
+public:
+	SProductionWizard() = default;
+	~SProductionWizard();
+
+	SLATE_BEGIN_ARGS(SProductionWizard) {}
+	SLATE_END_ARGS()
+
+	void Construct(const FArguments& InArgs);
+
+private:
+	/** Creates the list of productions used by the production list view widget*/
+	void UpdateProductionList();
+
+	/** Creates the menu on the left side of the production wizard */
+	TSharedRef<SWidget> MakeMenuPanel();
+
+	/** Creates the buttons on the bottom of the production wizard */
+	TSharedRef<SWidget> MakeButtonsPanel();
+
+	/** Functions for creating the panels for each of the production wizard's menu entries */
+	TSharedRef<SWidget> MakeRevisionControlPanel();
+	TSharedRef<SWidget> MakeProductionsPanel();
+	TSharedRef<SWidget> MakeSequencerSettingsPanel();
+	TSharedRef<SWidget> MakeNamingTokensPanel();
+	TSharedRef<SWidget> MakeAssetNamingPanel();
+	TSharedRef<SWidget> MakeFolderHierarchyPanel();
+
+	/** Creates the widget displaying the active production, used commonly by each of the production menu entry panels */
+	TSharedRef<SWidget> MakeActiveProductionWidget();
+
+	/** Generates a row in the menu list views */
+	TSharedRef<ITableRow> OnGenerateMenuRow(TSharedPtr<FProductionWizardMenuEntry> MenuEntry, const TSharedRef<STableViewBase>& OwnerTable);
+
+	/** Generates a row in the production list view */
+	TSharedRef<ITableRow> OnGenerateProductionRow(TSharedPtr<FCinematicProduction> InItem, const TSharedRef<STableViewBase>& OwnerTable);
+
+	/** Callback when one of the user menu entries is selected, which will cause its panel to be displayed on the right side of the production wizard */
+	void OnUserMenuSelectionChanged(TSharedPtr<FProductionWizardMenuEntry> SelectedEntry, ESelectInfo::Type SelectInfo);
+
+	/** Callback when one of the production menu entries is selected, which will cause its panel to be displayed on the right side of the production wizard */
+	void OnProductionMenuSelectionChanged(TSharedPtr<FProductionWizardMenuEntry> SelectedEntry, ESelectInfo::Type SelectInfo);
+
+	/** Callback when the production list view is finished regenerating its rows */
+	void OnProductionListItemsRebuilt();
+
+private:
+	/** Slots for the menu and menu entry content panels */
+	SSplitter::FSlot* MenuPanelSlot = nullptr;
+	SSplitter::FSlot* ContentPanelSlot = nullptr;
+
+	/** List items sources for the menu list views */
+	TArray<TSharedPtr<FProductionWizardMenuEntry>> UserMenuEntries;
+	TArray<TSharedPtr<FProductionWizardMenuEntry>> ProductionMenuEntries;
+
+	/** Menu list views, which allow the user to select a step in the production wizard */
+	TSharedPtr<SListView<TSharedPtr<FProductionWizardMenuEntry>>> UserMenuListView;
+	TSharedPtr<SListView<TSharedPtr<FProductionWizardMenuEntry>>> ProductionMenuListView;
+
+	/** List items sources for the production list view */
+	TArray<TSharedPtr<FCinematicProduction>> ProductionList;
+
+	/** Production list views, which allows the user to interact with one of the available productions in the Productions menu */
+	TSharedPtr<SListView<TSharedPtr<FCinematicProduction>>> ProductionListView;
+
+	/** Button belonging to one of the production list view rows that needs to be put into edit mode for the user to rename its production */
+	TSharedPtr<SProductionListButton> ProductionListButtonToRename;
+
+	/** Handle to the delegate which responds to changes in UProductionSettings list of productions */
+	FDelegateHandle ProductionListChangedHandle;
+};
diff --git a/Engine/Plugins/VirtualProduction/Shots/Source/ShotsEditor/Private/ShotsEditorModule.cpp b/Engine/Plugins/VirtualProduction/Shots/Source/ShotsEditor/Private/ShotsEditorModule.cpp
--- a/Engine/Plugins/VirtualProduction/Shots/Source/ShotsEditor/Private/ShotsEditorModule.cpp
+++ b/Engine/Plugins/VirtualProduction/Shots/Source/ShotsEditor/Private/ShotsEditorModule.cpp
@@ -2,9 +2,18 @@
 
 #include "Modules/ModuleManager.h"
 
+#include "Framework/Docking/TabManager.h"
 #include "ProductionSettings.h"
 #include "ProductionSettingsCustomization.h"
 #include "PropertyEditorModule.h"
+#include "SProductionWizard.h"
+#include "Widgets/Docking/SDockTab.h"
+#include "WorkspaceMenuStructure.h"
+#include "WorkspaceMenuStructureModule.h"
+
+#define LOCTEXT_NAMESPACE "ShotsEditorModule"
+
+static const FName ProductionWizardTabName("ProductionWizard");
 
 class FShotsEditorModule : public IModuleInterface
 {
@@ -14,6 +23,9 @@ public:
 	virtual void StartupModule() override;
 	virtual void ShutdownModule() override;
 	//~ End IModuleInterface interface
+
+private:
+	TSharedRef<SDockTab> SpawnProductionWizard(const FSpawnTabArgs& SpawnTabArgs);
 };
 
 void FShotsEditorModule::StartupModule()
@@ -24,12 +36,32 @@ void FShotsEditorModule::StartupModule()
 		UProductionSettings::StaticClass()->GetFName(),
 		FOnGetDetailCustomizationInstance::CreateStatic(&FProductionSettingsCustomization::MakeInstance)
 	);
+
+	// Add a new entry to the Tools->Cinematics menu to spawn the Production Wizard
+	FGlobalTabmanager::Get()->RegisterNomadTabSpawner(ProductionWizardTabName, FOnSpawnTab::CreateRaw(this, &FShotsEditorModule::SpawnProductionWizard))
+		.SetGroup(WorkspaceMenu::GetMenuStructure().GetLevelEditorCinematicsCategory())
+		.SetTooltipText(LOCTEXT("ProductionWizardTabTooltip", "Open the Production Wizard"))
+		.SetDisplayName(LOCTEXT("ProductionWizardTabTitle", "Production Wizard"))
+		.SetIcon(FSlateIcon(FAppStyle::GetAppStyleSetName(), "ProjectSettings.TabIcon"));
 }
 
 void FShotsEditorModule::ShutdownModule()
 {
 	FPropertyEditorModule& PropertyModule = FModuleManager::GetModuleChecked<FPropertyEditorModule>("PropertyEditor");
 	PropertyModule.UnregisterCustomClassLayout(UProductionSettings::StaticClass()->GetFName());
+
+	FGlobalTabmanager::Get()->UnregisterNomadTabSpawner(ProductionWizardTabName);
+}
+
+TSharedRef<SDockTab> FShotsEditorModule::SpawnProductionWizard(const FSpawnTabArgs& SpawnTabArgs)
+{
+	return SNew(SDockTab)
+		.TabRole(ETabRole::NomadTab)
+		[
+			SNew(SProductionWizard)
+		];
 }
 
 IMPLEMENT_MODULE(FShotsEditorModule, ShotsEditor);
+
+#undef LOCTEXT_NAMESPACE
diff --git a/Engine/Plugins/VirtualProduction/Shots/Source/ShotsEditor/Private/ShotsEditorStyle.cpp b/Engine/Plugins/VirtualProduction/Shots/Source/ShotsEditor/Private/ShotsEditorStyle.cpp
new file mode 100644
--- /dev/null
+++ b/Engine/Plugins/VirtualProduction/Shots/Source/ShotsEditor/Private/ShotsEditorStyle.cpp
@@ -0,0 +1,60 @@
+// Copyright Epic Games, Inc. All Rights Reserved.
+
+#include "ShotsEditorStyle.h"
+
+#include "Interfaces/IPluginManager.h"
+#include "Styling/CoreStyle.h"
+#include "Styling/SlateStyleMacros.h"
+#include "Styling/SlateTypes.h"
+#include "Styling/StyleColors.h"
+
+FName FShotsEditorStyle::StyleName("ShotsEditorStyle");
+
+FShotsEditorStyle& FShotsEditorStyle::Get()
+{
+	static FShotsEditorStyle Inst;
+	return Inst;
+}
+
+FShotsEditorStyle::FShotsEditorStyle()
+	: FSlateStyleSet(StyleName)
+{
+	// Icons
+	const FVector2D Icon16x16(16.0f, 16.0f);
+
+	const FString ContentDir = IPluginManager::Get().FindPlugin(TEXT("Shots"))->GetContentDir();
+	SetContentRoot(ContentDir);
+	SetCoreContentRoot(FPaths::EngineContentDir() / TEXT("Slate"));
+
+	Set("Icons.RevisionControl", new FSlateVectorImageBrush(RootToContentDir(TEXT("Slate/RevisionControl.svg")), Icon16x16));
+	Set("Icons.Productions", new FSlateVectorImageBrush(RootToContentDir(TEXT("Slate/Camera.svg")), Icon16x16));
+	Set("Icons.Sequencer", new FSlateVectorImageBrush(RootToContentDir(TEXT("Slate/Sequencer_16.svg")), Icon16x16));
+	Set("Icons.NamingTokens", new FSlateVectorImageBrush(RootToContentDir(TEXT("Slate/NamedSlot.svg")), Icon16x16));
+	Set("Icons.AssetNaming", new FSlateVectorImageBrush(RootToContentDir(TEXT("Slate/Rename.svg")), Icon16x16));
+	Set("Icons.Folder", new FSlateVectorImageBrush(RootToContentDir(TEXT("Slate/folder-closed.svg")), Icon16x16));
+	Set("Icons.Export", new FSlateVectorImageBrush(RootToContentDir(TEXT("Slate/ExportAll.svg")), Icon16x16));
+
+	// Fonts
+	Set("ProductionWizard.TitleFont", DEFAULT_FONT("Normal", 12));
+	Set("ProductionWizard.HeadingFont", DEFAULT_FONT("Bold", 14));
+
+	// Brushes
+	Set("ProductionWizard.PanelBackground", new FSlateRoundedBoxBrush(FStyleColors::Panel, 0.0f, FStyleColors::Recessed, 2.0f));
+	Set("ProductionWizard.RecessedBackground", new FSlateRoundedBoxBrush(FStyleColors::Recessed, 0.0f, FStyleColors::Panel, 2.0f));
+
+	// Buttons
+	const FButtonStyle RecessedButton = FButtonStyle()
+		.SetNormal(FSlateRoundedBoxBrush(FStyleColors::Recessed, 0.0f, FStyleColors::Recessed, 0.0f))
+		.SetHovered(FSlateRoundedBoxBrush(FStyleColors::Panel, 0.0f, FStyleColors::Panel, 0.0f))
+		.SetPressed(FSlateRoundedBoxBrush(FStyleColors::Dropdown, 0.0f, FStyleColors::Dropdown, 0.0f))
+		.SetDisabled(FSlateRoundedBoxBrush(FStyleColors::Recessed, 0.0f, FStyleColors::Recessed, 0.0f));
+
+	Set("ProductionWizard.RecessedButton", RecessedButton);
+
+	FSlateStyleRegistry::RegisterSlateStyle(*this);
+}
+
+FShotsEditorStyle::~FShotsEditorStyle()
+{
+	FSlateStyleRegistry::UnRegisterSlateStyle(*this);
+}
diff --git a/Engine/Plugins/VirtualProduction/Shots/Source/ShotsEditor/Private/ShotsEditorStyle.h b/Engine/Plugins/VirtualProduction/Shots/Source/ShotsEditor/Private/ShotsEditorStyle.h
new file mode 100644
--- /dev/null
+++ b/Engine/Plugins/VirtualProduction/Shots/Source/ShotsEditor/Private/ShotsEditorStyle.h
@@ -0,0 +1,21 @@
+// Copyright Epic Games, Inc. All Rights Reserved.
+
+#pragma once
+
+#include "Styling/SlateStyle.h"
+#include "Styling/SlateStyleRegistry.h"
+
+/** Slate style set that defines all the styles for the shots editor tools */
+class FShotsEditorStyle : public FSlateStyleSet
+{
+public:
+	static FName StyleName;
+
+	/** Access the singleton instance for this style set */
+	static FShotsEditorStyle& Get();
+
+private:
+
+	FShotsEditorStyle();
+	~FShotsEditorStyle();
+};
diff --git a/Engine/Plugins/VirtualProduction/Shots/Source/ShotsEditor/ShotsEditor.build.cs b/Engine/Plugins/VirtualProduction/Shots/Source/ShotsEditor/ShotsEditor.build.cs
--- a/Engine/Plugins/VirtualProduction/Shots/Source/ShotsEditor/ShotsEditor.build.cs
+++ b/Engine/Plugins/VirtualProduction/Shots/Source/ShotsEditor/ShotsEditor.build.cs
@@ -27,10 +27,13 @@ namespace UnrealBuildTool.Rules
 					"LevelSequenceEditor",
 					"MovieScene",
 					"MovieSceneTools",
+					"Projects",
 					"Slate",
 					"SlateCore",
 					"SourceControl",
+					"TimeManagement",
 					"UnrealEd",
+					"WorkspaceMenuStructure"
 				}
 			);
 		}
