ShotManagement: Fix CIS issue, multiple declarations

#rnx

diff --git a/Engine/Plugins/VirtualProduction/Shots/Source/ShotsEditor/Private/SProductionWizard.cpp b/Engine/Plugins/VirtualProduction/Shots/Source/ShotsEditor/Private/SProductionWizard.cpp
--- a/Engine/Plugins/VirtualProduction/Shots/Source/ShotsEditor/Private/SProductionWizard.cpp
+++ b/Engine/Plugins/VirtualProduction/Shots/Source/ShotsEditor/Private/SProductionWizard.cpp
@@ -829,9 +829,9 @@ void SProductionListButton::Construct(const FArguments& InArgs)
 	SButton::Construct(SButton::FArguments());
 	SetButtonStyle(&FShotsEditorStyle::Get().GetWidgetStyle<FButtonStyle>("ProductionWizard.RecessedButton"));
 
-	auto OnProductionButtonClicked = [](FGuid ProductionID) -> FReply
+	auto OnProductionButtonClicked = [](FGuid InProductionID) -> FReply
 	{
-		OnSetActiveProduction(ProductionID);
+		OnSetActiveProduction(InProductionID);
 		return FReply::Handled();
 	};
 
