ShotManagement: Add Production Function Library to provide BP/Python an interface to the Cinematic Production Settings
#rnx

diff --git a/Engine/Plugins/VirtualProduction/Shots/Source/ShotsEditor/Private/ProductionFunctionLibrary.cpp b/Engine/Plugins/VirtualProduction/Shots/Source/ShotsEditor/Private/ProductionFunctionLibrary.cpp
new file mode 100644
--- /dev/null
+++ b/Engine/Plugins/VirtualProduction/Shots/Source/ShotsEditor/Private/ProductionFunctionLibrary.cpp
@@ -0,0 +1,15 @@
+// Copyright Epic Games, Inc. All Rights Reserved.
+
+#include "ProductionFunctionLibrary.h"
+
+TArray<FCinematicProduction> UProductionFunctionLibrary::GetAllProductions()
+{
+	const UProductionSettings* const ProductionSettings = GetDefault<UProductionSettings>();
+	return ProductionSettings->Productions;
+}
+
+void UProductionFunctionLibrary::SetActiveProduction(FCinematicProduction Production)
+{
+	UProductionSettings* ProductionSettings = GetMutableDefault<UProductionSettings>();
+	ProductionSettings->SetActiveProduction(Production.ProductionID);
+}
diff --git a/Engine/Plugins/VirtualProduction/Shots/Source/ShotsEditor/Private/ProductionFunctionLibrary.h b/Engine/Plugins/VirtualProduction/Shots/Source/ShotsEditor/Private/ProductionFunctionLibrary.h
new file mode 100644
--- /dev/null
+++ b/Engine/Plugins/VirtualProduction/Shots/Source/ShotsEditor/Private/ProductionFunctionLibrary.h
@@ -0,0 +1,30 @@
+// Copyright Epic Games, Inc. All Rights Reserved.
+
+#pragma once
+
+#include "Kismet/BlueprintFunctionLibrary.h"
+#include "ProductionSettings.h"
+
+#include "ProductionFunctionLibrary.generated.h"
+
+/**
+ * Library of Blueprint/Python accessible functions to interface with the Cinematic Production Settings
+ */
+UCLASS()
+class UProductionFunctionLibrary : public UBlueprintFunctionLibrary
+{
+	GENERATED_BODY()
+
+public:
+
+	/** Returns an array of all available Cinematic Productions */
+	UFUNCTION(BlueprintPure, Category = "Shot Management")
+	static TArray<FCinematicProduction> GetAllProductions();
+
+	/** 
+	 * Sets the input Production as the current Active Production 
+	 * If no input is provided, the Active Production will be set to None
+	 */
+	UFUNCTION(BlueprintCallable, Category = "Shot Management")
+	static void SetActiveProduction(FCinematicProduction Production = FCinematicProduction());
+};
diff --git a/Engine/Plugins/VirtualProduction/Shots/Source/ShotsEditor/Private/ProductionSettings.cpp b/Engine/Plugins/VirtualProduction/Shots/Source/ShotsEditor/Private/ProductionSettings.cpp
--- a/Engine/Plugins/VirtualProduction/Shots/Source/ShotsEditor/Private/ProductionSettings.cpp
+++ b/Engine/Plugins/VirtualProduction/Shots/Source/ShotsEditor/Private/ProductionSettings.cpp
@@ -193,6 +193,9 @@ void UProductionSettings::PostEditChangeProperty(struct FPropertyChangedEvent& P
 			// Make sure the production name is unique
 			const int32 ArrayIndex = PropertyChangedEvent.GetArrayIndex(PropertyChangedEvent.MemberProperty->GetFName().ToString());
 			Productions[ArrayIndex].ProductionName = GetUniqueProductionName();
+
+			// Reset the active production pointer, in case the array was reallocated
+			ActiveProduction = Productions.FindByPredicate([this](const FCinematicProduction& Other) { return Other.ProductionName == ActiveProductionName; });
 		}
 		else
 		{
@@ -243,6 +246,9 @@ void UProductionSettings::AddProduction()
 {
 	const int32 ArrayIndex = Productions.Add(FCinematicProduction());
 
+	// Reset the active production pointer, in case the array was reallocated
+	ActiveProduction = Productions.FindByPredicate([this](const FCinematicProduction& Other) { return Other.ProductionName == ActiveProductionName; });
+
 	TryUpdateDefaultConfigFile();
 
 	ProductionListChangedDelegate.Broadcast();
@@ -472,4 +478,4 @@ FString UProductionSettings::GetUniqueProductionName() const
 	return WorkingName;
 }
 
-#undef LOCTEXT_NAMESPACE
\ No newline at end of file
+#undef LOCTEXT_NAMESPACE
diff --git a/Engine/Plugins/VirtualProduction/Shots/Source/ShotsEditor/Private/ProductionSettings.h b/Engine/Plugins/VirtualProduction/Shots/Source/ShotsEditor/Private/ProductionSettings.h
--- a/Engine/Plugins/VirtualProduction/Shots/Source/ShotsEditor/Private/ProductionSettings.h
+++ b/Engine/Plugins/VirtualProduction/Shots/Source/ShotsEditor/Private/ProductionSettings.h
@@ -22,7 +22,7 @@ enum class ESubsequencePriority : uint8
 /**
  * Collection of production settings to override the project/editor behavior
  */
-USTRUCT()
+USTRUCT(BlueprintType)
 struct FCinematicProduction
 {
 	GENERATED_BODY()
@@ -30,19 +30,19 @@ struct FCinematicProduction
 public:
 	FCinematicProduction();
 
-	UPROPERTY(config, EditAnywhere, Category = "Default")
+	UPROPERTY(config, EditAnywhere, BlueprintReadOnly, Category = "Default")
 	FString ProductionName;
 
 	/** The default frame rate to set for new Level Sequences */
-	UPROPERTY(config, EditAnywhere, Category = "Default")
+	UPROPERTY(config, EditAnywhere, BlueprintReadOnly, Category = "Default")
 	FFrameRate DefaultDisplayRate;
 
 	/** The default frame number (using the default frame rate) that new Level Sequences should start at */
-	UPROPERTY(config, EditAnywhere, Category = "Default")
+	UPROPERTY(config, EditAnywhere, BlueprintReadOnly, Category = "Default")
 	int32 DefaultStartFrame;
 
 	/** Controls whether subsequences override parent sequences, or vice versa */
-	UPROPERTY(config, EditAnywhere, Category = "Default")
+	UPROPERTY(config, EditAnywhere, BlueprintReadOnly, Category = "Default")
 	ESubsequencePriority SubsequencePriority = ESubsequencePriority::BottomUp;
 
 	/** Unique ID of the production */
