ShotManagement: Refactor production wizards into multiple widgets to house the functionality of each of the major panels.
#rnx

diff --git a/Engine/Plugins/VirtualProduction/Shots/Source/ShotsEditor/Private/SProductionWizard.cpp b/Engine/Plugins/VirtualProduction/Shots/Source/ShotsEditor/Private/SProductionWizard.cpp
deleted file mode 100644
diff --git a/Engine/Plugins/VirtualProduction/Shots/Source/ShotsEditor/Private/SProductionWizard.h b/Engine/Plugins/VirtualProduction/Shots/Source/ShotsEditor/Private/SProductionWizard.h
deleted file mode 100644
diff --git a/Engine/Plugins/VirtualProduction/Shots/Source/ShotsEditor/Private/ShotsEditorModule.cpp b/Engine/Plugins/VirtualProduction/Shots/Source/ShotsEditor/Private/ShotsEditorModule.cpp
--- a/Engine/Plugins/VirtualProduction/Shots/Source/ShotsEditor/Private/ShotsEditorModule.cpp
+++ b/Engine/Plugins/VirtualProduction/Shots/Source/ShotsEditor/Private/ShotsEditorModule.cpp
@@ -6,7 +6,7 @@
 #include "ProductionSettings.h"
 #include "ProductionSettingsCustomization.h"
 #include "PropertyEditorModule.h"
-#include "SProductionWizard.h"
+#include "UI/SProductionWizard.h"
 #include "Widgets/Docking/SDockTab.h"
 #include "WorkspaceMenuStructure.h"
 #include "WorkspaceMenuStructureModule.h"
diff --git a/Engine/Plugins/VirtualProduction/Shots/Source/ShotsEditor/Private/UI/SActiveProductionCombo.cpp b/Engine/Plugins/VirtualProduction/Shots/Source/ShotsEditor/Private/UI/SActiveProductionCombo.cpp
new file mode 100644
--- /dev/null
+++ b/Engine/Plugins/VirtualProduction/Shots/Source/ShotsEditor/Private/UI/SActiveProductionCombo.cpp
@@ -0,0 +1,88 @@
+// Copyright Epic Games, Inc. All Rights Reserved.
+
+#include "SActiveProductionCombo.h"
+
+#include "Framework/MultiBox/MultiBoxBuilder.h"
+#include "ProductionSettings.h"
+#include "Widgets/Input/SComboButton.h"
+#include "Widgets/Layout/SBorder.h"
+#include "Widgets/SBoxPanel.h"
+
+#define LOCTEXT_NAMESPACE "SActiveProductionCombo"
+
+void SActiveProductionCombo::Construct(const FArguments& InArgs)
+{
+	static auto GetActiveProductionName = []() -> FText
+		{
+			const UProductionSettings* const ProductionSettings = GetDefault<UProductionSettings>();
+			return FText::FromString(ProductionSettings->ActiveProductionName);
+		};
+
+	static auto OnSetActiveProduction = [](FGuid ProductionID)
+		{
+			UProductionSettings* ProductionSettings = GetMutableDefault<UProductionSettings>();
+			ProductionSettings->SetActiveProduction(ProductionID);
+		};
+
+	auto BuildMenu = []() -> TSharedRef<SWidget>
+		{
+			FMenuBuilder MenuBuilder(true, nullptr);
+
+			// None option, allowing the user to set no active production
+			MenuBuilder.AddMenuEntry(
+				FText::FromName(NAME_None),
+				FText::FromName(NAME_None),
+				FSlateIcon(),
+				FUIAction(FExecuteAction::CreateLambda(OnSetActiveProduction, FGuid())),
+				NAME_None,
+				EUserInterfaceActionType::Button
+			);
+
+			const UProductionSettings* const ProductionSettings = GetDefault<UProductionSettings>();
+			for (const FCinematicProduction& Production : ProductionSettings->Productions)
+			{
+				MenuBuilder.AddMenuEntry(
+					FText::FromString(Production.ProductionName),
+					FText::FromString(Production.ProductionName),
+					FSlateIcon(),
+					FUIAction(FExecuteAction::CreateLambda(OnSetActiveProduction, Production.ProductionID)),
+					NAME_None,
+					EUserInterfaceActionType::Button
+				);
+			}
+
+			return MenuBuilder.MakeWidget();
+		};
+
+	ChildSlot
+		[
+			SNew(SBorder)
+				.BorderImage(FAppStyle::Get().GetBrush("Brushes.Panel"))
+				.Padding(16.0f, 8.0f)
+				[
+					SNew(SHorizontalBox)
+
+						+ SHorizontalBox::Slot()
+						.AutoWidth()
+						.VAlign(VAlign_Center)
+						.Padding(0.0f, 0.0f, 8.0f, 0.0f)
+						[
+							SNew(STextBlock).Text(LOCTEXT("ActiveProductionText", "Active Production"))
+						]
+
+						+ SHorizontalBox::Slot()
+						.FillWidth(1.0f)
+						.VAlign(VAlign_Center)
+						[
+							SNew(SComboButton)
+								.OnGetMenuContent_Lambda(BuildMenu)
+								.ButtonContent()
+								[
+									SNew(STextBlock).Text_Lambda(GetActiveProductionName)
+								]
+						]
+				]
+		];
+}
+
+#undef LOCTEXT_NAMESPACE
diff --git a/Engine/Plugins/VirtualProduction/Shots/Source/ShotsEditor/Private/UI/SActiveProductionCombo.h b/Engine/Plugins/VirtualProduction/Shots/Source/ShotsEditor/Private/UI/SActiveProductionCombo.h
new file mode 100644
--- /dev/null
+++ b/Engine/Plugins/VirtualProduction/Shots/Source/ShotsEditor/Private/UI/SActiveProductionCombo.h
@@ -0,0 +1,19 @@
+// Copyright Epic Games, Inc. All Rights Reserved.
+
+#pragma once
+
+#include "Widgets/SCompoundWidget.h"
+
+#include "Widgets/DeclarativeSyntaxSupport.h"
+
+/**
+ * Combo button widget for choosing an available production to be the Active production
+ */
+class SActiveProductionCombo : public SCompoundWidget
+{
+public:
+	SLATE_BEGIN_ARGS(SActiveProductionCombo) {}
+	SLATE_END_ARGS()
+
+	void Construct(const FArguments& InArgs);
+};
diff --git a/Engine/Plugins/VirtualProduction/Shots/Source/ShotsEditor/Private/UI/SProductionListPanel.cpp b/Engine/Plugins/VirtualProduction/Shots/Source/ShotsEditor/Private/UI/SProductionListPanel.cpp
new file mode 100644
--- /dev/null
+++ b/Engine/Plugins/VirtualProduction/Shots/Source/ShotsEditor/Private/UI/SProductionListPanel.cpp
@@ -0,0 +1,411 @@
+// Copyright Epic Games, Inc. All Rights Reserved.
+
+#include "SProductionListPanel.h"
+
+#include "Framework/Application/MenuStack.h"
+#include "Framework/Application/SlateApplication.h"
+#include "Framework/MultiBox/MultiBoxBuilder.h"
+#include "Layout/WidgetPath.h"
+#include "ProductionSettings.h"
+#include "ShotsEditorStyle.h"
+#include "Styling/StyleColors.h"
+#include "Widgets/Layout/SBorder.h"
+#include "Widgets/SBoxPanel.h"
+
+#include "SProductionListPanel.h"
+
+#define LOCTEXT_NAMESPACE "SProductionListPanel"
+
+void SProductionListPanel::Construct(const FArguments& InArgs)
+{
+	// Subscribe to be notified when the Production Settings list of productions has changed (for example, if a production was added/removed)
+	UProductionSettings* ProductionSettings = GetMutableDefault<UProductionSettings>();
+	ProductionListChangedHandle = ProductionSettings->OnProductionListChanged().AddSP(this, &SProductionListPanel::UpdateProductionList);
+
+	// Build the list of productions for the list view before building the widget for the first time
+	UpdateProductionList();
+
+	auto OnCreateNewProductionClicked = []() -> FReply
+		{
+			UProductionSettings* ProductionSettings = GetMutableDefault<UProductionSettings>();
+			ProductionSettings->AddProduction();
+			return FReply::Handled();
+		};
+
+	ChildSlot
+		[
+			SNew(SBorder)
+				.BorderImage(FAppStyle::Get().GetBrush("Brushes.Panel"))
+				.Padding(16.0f)
+				[
+					SNew(SVerticalBox)
+
+						// Title
+						+ SVerticalBox::Slot()
+						.AutoHeight()
+						.Padding(0.0f, 0.0f, 0.0f, 8.0f)
+						[
+							SNew(STextBlock)
+								.Text(LOCTEXT("ProductionsTitle", "User Setup"))
+								.Font(FShotsEditorStyle::Get().GetFontStyle("ProductionWizard.TitleFont"))
+						]
+
+						// Heading
+						+ SVerticalBox::Slot()
+						.AutoHeight()
+						.Padding(0.0f, 0.0f, 0.0f, 8.0f)
+						[
+							SNew(STextBlock)
+								.Text(LOCTEXT("ProductionsHeading", "Productions"))
+								.Font(FShotsEditorStyle::Get().GetFontStyle("ProductionWizard.HeadingFont"))
+						]
+
+						// Info Text
+						+ SVerticalBox::Slot()
+						.AutoHeight()
+						.Padding(0.0f, 0.0f, 0.0f, 20.0f)
+						[
+							SNew(STextBlock)
+								.Text(LOCTEXT("ProductionsInfoText", "You can use this Unreal project for multiple shows, and assign different settings for each show.\n"
+									"Create a production here for each show, then choose one to be the active production."))
+						]
+
+						// Create / Import Buttons
+						+ SVerticalBox::Slot()
+						.AutoHeight()
+						.Padding(0.0f, 0.0f, 0.0f, 4.0f)
+						[
+							SNew(SHorizontalBox)
+
+								+ SHorizontalBox::Slot()
+								.AutoWidth()
+								[
+									SNew(SButton)
+										.ContentPadding(FMargin(2.0f))
+										.OnClicked_Lambda(OnCreateNewProductionClicked)
+										[
+											SNew(SHorizontalBox)
+
+												+ SHorizontalBox::Slot()
+												.AutoWidth()
+												.Padding(0.0f, 0.0f, 4.0f, 0.0f)
+												[
+													SNew(SImage)
+														.Image(FAppStyle::Get().GetBrush("Icons.Plus"))
+														.ColorAndOpacity(FStyleColors::AccentGreen)
+												]
+
+
+												+ SHorizontalBox::Slot()
+												.AutoWidth()
+												[
+													SNew(STextBlock).Text(LOCTEXT("CreateNewProductionButton", "Create a New Production"))
+												]
+										]
+								]
+
+								+ SHorizontalBox::Slot()
+								.AutoWidth()
+								[
+									// TODO: Add support for importing a production when this button is clicked
+									SNew(SButton)
+										.ContentPadding(FMargin(2.0f))
+										[
+											SNew(SHorizontalBox)
+
+												+ SHorizontalBox::Slot()
+												.AutoWidth()
+												.Padding(0.0f, 0.0f, 4.0f, 0.0f)
+												[
+													SNew(SImage)
+														.Image(FAppStyle::Get().GetBrush("Icons.Import"))
+												]
+
+												+ SHorizontalBox::Slot()
+												.AutoWidth()
+												[
+													SNew(STextBlock)
+														.Text(LOCTEXT("ImportProductionButton", "Import Production"))
+														.ToolTipText(LOCTEXT("ImportProductionButtonToolTip", "Not yet implemented"))
+												]
+										]
+								]
+						]
+
+						// Production List View
+						+ SVerticalBox::Slot()
+						.FillHeight(1.0f)
+						[
+							SNew(SBorder)
+								.BorderImage(FAppStyle::Get().GetBrush("Brushes.Recessed"))
+								.Padding(8.0f)
+								[
+									SAssignNew(ProductionListView, SListView<TSharedPtr<FCinematicProduction>>)
+										.ListItemsSource(&ProductionList)
+										.SelectionMode(ESelectionMode::None)
+										.OnGenerateRow(this, &SProductionListPanel::OnGenerateProductionRow)
+										.OnItemsRebuilt(this, &SProductionListPanel::OnProductionListItemsRebuilt)
+								]
+						]
+				]
+		];
+}
+
+SProductionListPanel::~SProductionListPanel()
+{
+	UProductionSettings* ProductionSettings = GetMutableDefault<UProductionSettings>();
+	ProductionSettings->OnProductionListChanged().Remove(ProductionListChangedHandle);
+}
+
+void SProductionListPanel::UpdateProductionList()
+{
+	ProductionList.Empty();
+
+	const UProductionSettings* const ProductionSettings = GetDefault<UProductionSettings>();
+	for (const FCinematicProduction& Production : ProductionSettings->Productions)
+	{
+		ProductionList.Add(MakeShared<FCinematicProduction>(Production));
+	}
+
+	if (ProductionListView)
+	{
+		ProductionListView->RequestListRefresh();
+	}
+}
+
+void SProductionListPanel::OnProductionListItemsRebuilt()
+{
+	// If there was a production with no name, then ProductionListButtonToRename will be valid, and it can be put into edit mode so the user can rename the production.
+	if (ProductionListButtonToRename)
+	{
+		ProductionListButtonToRename->EnterEditMode();
+		ProductionListButtonToRename = nullptr;
+	}
+}
+
+TSharedRef<ITableRow> SProductionListPanel::OnGenerateProductionRow(TSharedPtr<FCinematicProduction> InItem, const TSharedRef<STableViewBase>& OwnerTable)
+{
+	auto OnDeleteProduction = [](FGuid ProductionID) -> FReply
+		{
+			UProductionSettings* ProductionSettings = GetMutableDefault<UProductionSettings>();
+			ProductionSettings->DeleteProduction(ProductionID);
+			return FReply::Handled();
+		};
+
+	TSharedRef<SProductionListButton> ListButton = SNew(SProductionListButton)
+		.ProductionID(InItem->ProductionID);
+
+	// If this production was newly created, it will not have a name yet, so the user should be prompted to give it one.
+	// However, the production's list button cannot be put into edit mode and focused until after the widget is created.
+	// Therefore, the widget is saved, and will be put into edit mode after the items for this ListView are finished being rebuilt. 
+	if (InItem->ProductionName.IsEmpty())
+	{
+		ProductionListButtonToRename = ListButton;
+	}
+
+	return SNew(STableRow<TSharedPtr<FName>>, OwnerTable)
+		.ShowSelection(true)
+		.Content()
+		[
+			SNew(SHorizontalBox)
+
+				+ SHorizontalBox::Slot()
+				.FillWidth(1.0f)
+				[
+					ListButton
+				]
+
+				+ SHorizontalBox::Slot()
+				.AutoWidth()
+				.VAlign(VAlign_Center)
+				[
+					// TODO: Add support for exporting a production when this button is clicked
+					SNew(SButton)
+						.ButtonStyle(FShotsEditorStyle::Get(), "ProductionWizard.RecessedButton")
+						[
+							SNew(SImage).Image(FShotsEditorStyle::Get().GetBrush("Icons.Export"))
+						]
+				]
+
+				+ SHorizontalBox::Slot()
+				.AutoWidth()
+				.VAlign(VAlign_Center)
+				[
+					SNew(SButton)
+						.ButtonStyle(FShotsEditorStyle::Get(), "ProductionWizard.RecessedButton")
+						.OnClicked_Lambda(OnDeleteProduction, InItem->ProductionID)
+						[
+							SNew(SImage).Image(FAppStyle::Get().GetBrush("Icons.Delete"))
+						]
+				]
+		];
+}
+
+SProductionListButton::SProductionListButton()
+	: ProductionID(*this)
+{
+}
+
+void SProductionListButton::Construct(const FArguments& InArgs)
+{
+	ProductionID.Assign(*this, InArgs._ProductionID);
+
+	SButton::Construct(SButton::FArguments());
+	SetButtonStyle(&FShotsEditorStyle::Get().GetWidgetStyle<FButtonStyle>("ProductionWizard.RecessedButton"));
+
+	auto OnProductionButtonClicked = [](FGuid InProductionID) -> FReply
+		{
+			UProductionSettings* ProductionSettings = GetMutableDefault<UProductionSettings>();
+			if (ProductionSettings->IsActiveProduction(InProductionID))
+			{
+				ProductionSettings->SetActiveProduction(FGuid());
+			}
+			else
+			{
+				ProductionSettings->SetActiveProduction(InProductionID);
+			}
+
+			return FReply::Handled();
+		};
+
+	SetOnClicked(FOnClicked::CreateLambda(OnProductionButtonClicked, ProductionID.Get()));
+
+	// Returns the color of the checkmark icon based on the active production and button hover state
+	auto GetCheckMarkColor = [this]() -> FSlateColor
+		{
+			const UProductionSettings* const ProductionSettings = GetDefault<UProductionSettings>();
+			if (ProductionSettings->IsActiveProduction(ProductionID.Get()))
+			{
+				return FStyleColors::AccentGreen;
+			}
+			else if (bIsHovered)
+			{
+				return FStyleColors::Foreground;
+			}
+			return FStyleColors::Transparent;
+		};
+
+	// Returns the production name based on the ProductionID assigned to this button
+	auto GetProductionName = [this]() -> FText
+		{
+			const UProductionSettings* const ProductionSettings = GetDefault<UProductionSettings>();
+			if (const FCinematicProduction* Production = ProductionSettings->GetProduction(ProductionID.Get()))
+			{
+				if (Production->ProductionName.IsEmpty())
+				{
+					return FText::FromString(ProductionSettings->GetUniqueProductionName());
+				}
+				return FText::FromString(Production->ProductionName);
+			}
+			return FText::GetEmpty();
+		};
+
+	// Renames this button's production
+	auto OnRenameFinished = [this](const FText& InText, ETextCommit::Type InCommitType)
+		{
+			UProductionSettings* ProductionSettings = GetMutableDefault<UProductionSettings>();
+			ProductionSettings->RenameProduction(ProductionID.Get(), InText.ToString());
+		};
+
+	auto OnVerifyRename = [this](const FText& InText, FText& OutErrorMessage) -> bool
+		{
+			const UProductionSettings* const ProductionSettings = GetDefault<UProductionSettings>();
+
+			if (InText.IsEmpty())
+			{
+				OutErrorMessage = LOCTEXT("EmptyNameErrorMessage", "Please provide a name for the production");
+				return false;
+			}
+
+			bool bIsNameValid = true;
+			for (const FCinematicProduction& Production : ProductionSettings->Productions)
+			{
+				if (Production.ProductionID == ProductionID.Get())
+				{
+					continue;
+				}
+
+				if (InText.ToString() == Production.ProductionName)
+				{
+					bIsNameValid = false;
+					OutErrorMessage = LOCTEXT("ExistingNameErrorMessage", "A production already exists with this name");
+				}
+			}
+			return bIsNameValid;
+		};
+
+	ChildSlot
+		[
+			SNew(SHorizontalBox)
+
+				+ SHorizontalBox::Slot()
+				.AutoWidth()
+				.VAlign(VAlign_Center)
+				.Padding(0.0f, 0.0f, 8.0f, 0.0f)
+				[
+					SNew(SImage)
+						.Image(FAppStyle::Get().GetBrush("Icons.Check"))
+						.ColorAndOpacity_Lambda(GetCheckMarkColor)
+				]
+
+				+ SHorizontalBox::Slot()
+				.FillWidth(1.0f)
+				[
+					SAssignNew(EditableTextBlock, SInlineEditableTextBlock)
+						.Text_Lambda(GetProductionName)
+						.OnVerifyTextChanged_Lambda(OnVerifyRename)
+						.OnTextCommitted_Lambda(OnRenameFinished)
+						.IsSelected_Lambda([]() { return false; }) // Disable double-select to rename
+				]
+		];
+}
+
+void SProductionListButton::EnterEditMode()
+{
+	EditableTextBlock->EnterEditingMode();
+}
+
+FReply SProductionListButton::OnMouseButtonDown(const FGeometry& MyGeometry, const FPointerEvent& MouseEvent)
+{
+	// Create the context menu to be launched on right mouse click
+	if (MouseEvent.GetEffectingButton() == EKeys::RightMouseButton)
+	{
+		FMenuBuilder MenuBuilder(true, nullptr);
+		MenuBuilder.AddMenuEntry(
+			LOCTEXT("RenameProduction", "Rename"),
+			LOCTEXT("RenameProductionToolTip", "Rename production"),
+			FSlateIcon(FShotsEditorStyle::StyleName, "Icons.AssetNaming"),
+			FUIAction(FExecuteAction::CreateSP(this, &SProductionListButton::EnterEditMode)),
+			NAME_None,
+			EUserInterfaceActionType::Button
+		);
+
+		FWidgetPath WidgetPath = MouseEvent.GetEventPath() != nullptr ? *MouseEvent.GetEventPath() : FWidgetPath();
+
+		FSlateApplication::Get().PushMenu(
+			AsShared(),
+			WidgetPath,
+			MenuBuilder.MakeWidget(),
+			MouseEvent.GetScreenSpacePosition(),
+			FPopupTransitionEffect(FPopupTransitionEffect::ContextMenu)
+		);
+
+		return FReply::Handled();
+	}
+
+	return SButton::OnMouseButtonDown(MyGeometry, MouseEvent);
+}
+
+void SProductionListButton::OnMouseEnter(const FGeometry& MyGeometry, const FPointerEvent& MouseEvent)
+{
+	SButton::OnMouseEnter(MyGeometry, MouseEvent);
+	bIsHovered = true;
+}
+
+void SProductionListButton::OnMouseLeave(const FPointerEvent& MouseEvent)
+{
+	SButton::OnMouseLeave(MouseEvent);
+	bIsHovered = false;
+}
+
+#undef LOCTEXT_NAMESPACE
diff --git a/Engine/Plugins/VirtualProduction/Shots/Source/ShotsEditor/Private/UI/SProductionListPanel.h b/Engine/Plugins/VirtualProduction/Shots/Source/ShotsEditor/Private/UI/SProductionListPanel.h
new file mode 100644
--- /dev/null
+++ b/Engine/Plugins/VirtualProduction/Shots/Source/ShotsEditor/Private/UI/SProductionListPanel.h
@@ -0,0 +1,88 @@
+// Copyright Epic Games, Inc. All Rights Reserved.
+
+#pragma once
+
+#include "Widgets/SCompoundWidget.h"
+
+#include "Widgets/Input/SButton.h"
+#include "Widgets/Text/SInlineEditableTextBlock.h"
+#include "Widgets/Views/SListView.h"
+
+struct FCinematicProduction;
+
+/**
+ * Button used in the Productions panel, containing functionality for a single FCinematicProduction.
+ * This button supports right-click (for renaming) and will show an editable text box for the user to rename a button.
+ */
+class SProductionListButton : public SButton
+{
+public:
+	SProductionListButton();
+
+	SLATE_BEGIN_ARGS(SProductionListButton)
+		: _ProductionID()
+		{ }
+
+		/** ID of the production that this button controls */
+		SLATE_ATTRIBUTE(FGuid, ProductionID)
+
+	SLATE_END_ARGS()
+
+	void Construct(const FArguments& InArgs);
+
+	/** OnMouseButtonDown is overridden because the base class SButton ignores right-clicks */
+	virtual FReply OnMouseButtonDown(const FGeometry& MyGeometry, const FPointerEvent& MouseEvent) override;
+	virtual void OnMouseEnter(const FGeometry& MyGeometry, const FPointerEvent& MouseEvent) override;
+	virtual void OnMouseLeave(const FPointerEvent& MouseEvent) override;
+
+	/** Put the editable text block into edit mode */
+	void EnterEditMode();
+
+private:
+	/** ID of the production that this button controls */
+	TSlateAttribute<FGuid, EInvalidateWidgetReason::Paint> ProductionID;
+
+	/** User-editable text widget for renaming the production */
+	TSharedPtr<SInlineEditableTextBlock> EditableTextBlock;
+
+	/** Hover state, used to determine color/opacity of the button icon */
+	bool bIsHovered = false;
+};
+
+/**
+ * UI for the Productions panel of the Production Wizard
+ */
+class SProductionListPanel : public SCompoundWidget
+{
+public:
+	SProductionListPanel() = default;
+	~SProductionListPanel();
+
+	SLATE_BEGIN_ARGS(SProductionListPanel) {}
+	SLATE_END_ARGS()
+
+	void Construct(const FArguments& InArgs);
+
+private:
+	/** Creates the list of productions used by the production list view widget*/
+	void UpdateProductionList();
+
+	/** Generates a row in the production list view */
+	TSharedRef<ITableRow> OnGenerateProductionRow(TSharedPtr<FCinematicProduction> InItem, const TSharedRef<STableViewBase>& OwnerTable);
+
+	/** Callback when the production list view is finished regenerating its rows */
+	void OnProductionListItemsRebuilt();
+
+private:
+	/** List items sources for the production list view */
+	TArray<TSharedPtr<FCinematicProduction>> ProductionList;
+
+	/** Production list views, which allows the user to interact with one of the available productions in the Productions menu */
+	TSharedPtr<SListView<TSharedPtr<FCinematicProduction>>> ProductionListView;
+
+	/** Button belonging to one of the production list view rows that needs to be put into edit mode for the user to rename its production */
+	TSharedPtr<SProductionListButton> ProductionListButtonToRename;
+
+	/** Handle to the delegate which responds to changes in UProductionSettings list of productions */
+	FDelegateHandle ProductionListChangedHandle;
+};
diff --git a/Engine/Plugins/VirtualProduction/Shots/Source/ShotsEditor/Private/UI/SProductionWizard.cpp b/Engine/Plugins/VirtualProduction/Shots/Source/ShotsEditor/Private/UI/SProductionWizard.cpp
new file mode 100644
--- /dev/null
+++ b/Engine/Plugins/VirtualProduction/Shots/Source/ShotsEditor/Private/UI/SProductionWizard.cpp
@@ -0,0 +1,320 @@
+// Copyright Epic Games, Inc. All Rights Reserved.
+
+#include "SProductionWizard.h"
+
+#include "ProductionSettings.h"
+#include "ShotsEditorStyle.h"
+#include "UI/SProductionListPanel.h"
+#include "UI/SSequencerSettingsPanel.h"
+#include "Widgets/Layout/SBorder.h"
+#include "Widgets/Layout/SSeparator.h"
+#include "Widgets/Layout/SSplitter.h"
+#include "Widgets/SBoxPanel.h"
+
+#define LOCTEXT_NAMESPACE "SProductionWizard"
+
+void SProductionWizard::Construct(const FArguments& InArgs)
+{
+	ChildSlot
+		[
+			SNew(SBorder)
+				.Padding(0.f)
+				.BorderImage(FShotsEditorStyle::Get().GetBrush("ProductionWizard.PanelBackground"))
+				[
+					SNew(SVerticalBox)
+
+					+ SVerticalBox::Slot()
+						[
+							SNew(SSplitter)
+								.Orientation(Orient_Horizontal)
+								.PhysicalSplitterHandleSize(2.0f)
+
+							+ SSplitter::Slot()
+								.Value(0.25f)
+								.MinSize(192.0f)
+								.Expose(MenuPanelSlot)
+
+							+ SSplitter::Slot()
+								.Value(0.75f)
+								.Expose(ContentPanelSlot)
+						]
+
+					+ SVerticalBox::Slot()
+						.AutoHeight()
+						[
+							SNew(SSeparator)
+								.Orientation(Orient_Horizontal)
+								.Thickness(2.0f)
+						]
+
+					+ SVerticalBox::Slot()
+						.AutoHeight()
+						[
+							MakeButtonsPanel()
+						]
+				]
+		];
+
+	MenuPanelSlot->AttachWidget(MakeMenuPanel());
+}
+
+TSharedRef<SWidget> SProductionWizard::MakeButtonsPanel()
+{
+	return SNew(SBorder)
+		.BorderImage(FAppStyle::Get().GetBrush("Brushes.Panel"))
+		.Padding(16.0f)
+		.HAlign(HAlign_Right)
+		[
+			SNew(SHorizontalBox)
+
+			+ SHorizontalBox::Slot()
+				.MinWidth(118.0f)
+				.MaxWidth(118.0f)
+				.VAlign(VAlign_Center)
+				.Padding(0.0f, 0.0f, 8.0f, 0.0f)
+				[
+					SNew(SButton)
+						.Text(LOCTEXT("BackButton", "Back"))
+						.HAlign(HAlign_Center)
+				]
+
+			+ SHorizontalBox::Slot()
+				.MinWidth(118.0f)
+				.MaxWidth(118.0f)
+				.VAlign(VAlign_Center)
+				[
+					SNew(SButton)
+						.Text(LOCTEXT("NextButton", "Next"))
+						.HAlign(HAlign_Center)
+				]
+		];
+}
+
+TSharedRef<SWidget> SProductionWizard::MakeMenuPanel()
+{
+	// Build the Revision Control menu entry, including its content panel
+	TSharedPtr<FProductionWizardMenuEntry> RevisionControlMenuEntry = MakeShared<FProductionWizardMenuEntry>();
+	RevisionControlMenuEntry->Label = LOCTEXT("RevisionControlMenuLabel", "Revision Control");
+	RevisionControlMenuEntry->Icon = FSlateIcon(FShotsEditorStyle::StyleName, "Icons.RevisionControl");
+	RevisionControlMenuEntry->Panel = MakeRevisionControlPanel();
+
+	// Build the Productions menu entry, including its content panel
+	TSharedPtr<FProductionWizardMenuEntry> ProductionMenuEntry = MakeShared<FProductionWizardMenuEntry>();
+	ProductionMenuEntry->Label = LOCTEXT("ProductionsMenuLabel", "Productions");
+	ProductionMenuEntry->Icon = FSlateIcon(FShotsEditorStyle::StyleName, "Icons.Productions");
+	ProductionMenuEntry->Panel = SNew(SProductionListPanel);
+
+	UserMenuEntries.Add(MoveTemp(RevisionControlMenuEntry));
+	UserMenuEntries.Add(MoveTemp(ProductionMenuEntry));
+
+	// Build the Sequencer Settings menu entry, including its content panel
+	TSharedPtr<FProductionWizardMenuEntry> SequencerSettingsMenuEntry = MakeShared<FProductionWizardMenuEntry>();
+	SequencerSettingsMenuEntry->Label = LOCTEXT("SequencerSettingsMenuLabel", "Sequencer Settings");
+	SequencerSettingsMenuEntry->Icon = FSlateIcon(FShotsEditorStyle::StyleName, "Icons.Sequencer");
+	SequencerSettingsMenuEntry->Panel = SNew(SSequencerSettingsPanel);
+
+	// Build the Naming Tokens menu entry, including its content panel
+	TSharedPtr<FProductionWizardMenuEntry> NamingTokensMenuEntry = MakeShared<FProductionWizardMenuEntry>();
+	NamingTokensMenuEntry->Label = LOCTEXT("NamingTokensMenuLabel", "Naming Tokens");
+	NamingTokensMenuEntry->Icon = FSlateIcon(FShotsEditorStyle::StyleName, "Icons.NamingTokens");
+	NamingTokensMenuEntry->Panel = MakeNamingTokensPanel();
+
+	// Build the Asset Naming menu entry, including its content panel
+	TSharedPtr<FProductionWizardMenuEntry> AssetNamingMenuEntry = MakeShared<FProductionWizardMenuEntry>();
+	AssetNamingMenuEntry->Label = LOCTEXT("AssetNamingMenuLabel", "Asset Naming");
+	AssetNamingMenuEntry->Icon = FSlateIcon(FShotsEditorStyle::StyleName, "Icons.AssetNaming");
+	AssetNamingMenuEntry->Panel = MakeAssetNamingPanel();
+
+	// Build the Folder Hierarchy menu entry, including its content panel
+	TSharedPtr<FProductionWizardMenuEntry> FolderHierarchyMenuEntry = MakeShared<FProductionWizardMenuEntry>();
+	FolderHierarchyMenuEntry->Label = LOCTEXT("FolderHierarchyMenuLabel", "Folder Hierarchy");
+	FolderHierarchyMenuEntry->Icon = FSlateIcon(FShotsEditorStyle::StyleName, "Icons.Folder");
+	FolderHierarchyMenuEntry->Panel = MakeFolderHierarchyPanel();
+
+	ProductionMenuEntries.Add(SequencerSettingsMenuEntry);
+	ProductionMenuEntries.Add(NamingTokensMenuEntry);
+	ProductionMenuEntries.Add(AssetNamingMenuEntry);
+	ProductionMenuEntries.Add(FolderHierarchyMenuEntry);
+
+	// Construct the list view for the user menu entries
+	UserMenuListView = SNew(SListView<TSharedPtr<FProductionWizardMenuEntry>>)
+		.ListItemsSource(&UserMenuEntries)
+		.OnGenerateRow(this, &SProductionWizard::OnGenerateMenuRow)
+		.OnSelectionChanged(this, &SProductionWizard::OnUserMenuSelectionChanged);
+
+	// Construct the list view for the production menu entries
+	ProductionMenuListView = SNew(SListView<TSharedPtr<FProductionWizardMenuEntry>>)
+		.ListItemsSource(&ProductionMenuEntries)
+		.OnGenerateRow(this, &SProductionWizard::OnGenerateMenuRow)
+		.OnSelectionChanged(this, &SProductionWizard::OnProductionMenuSelectionChanged);
+
+	// The initial selection upon opening the wizard is the first user menu entry (revision control)
+	UserMenuListView->SetSelection(UserMenuEntries[0]);
+	ContentPanelSlot->AttachWidget(UserMenuEntries[0]->Panel.ToSharedRef());
+
+	return SNew(SBorder)
+		.BorderImage(FShotsEditorStyle::Get().GetBrush("ProductionWizard.RecessedBackground"))
+		.Padding(16.0f)
+		[
+			SNew(SVerticalBox)
+
+			+ SVerticalBox::Slot()
+				.AutoHeight()
+				.Padding(0.0f, 0.0f, 0.0f, 8.0f)
+				[
+					SNew(STextBlock)
+						.Text(LOCTEXT("UserSetupMenuGroup", "User Setup"))
+						.Font(FAppStyle::Get().GetFontStyle("NormalFontBold"))
+				]
+
+			+ SVerticalBox::Slot()
+				.AutoHeight()
+				.Padding(0.0f, 0.0f, 0.0f, 8.0f)
+				[
+					UserMenuListView.ToSharedRef()
+				]
+
+			+ SVerticalBox::Slot()
+				.AutoHeight()
+				.Padding(0.0f, 0.0f, 0.0f, 8.0f)
+				[
+					SNew(STextBlock)
+						.Text(LOCTEXT("ProductionSettingsMenuGroup", "Production Settings"))
+						.Font(FAppStyle::Get().GetFontStyle("NormalFontBold"))
+				]
+
+			+ SVerticalBox::Slot()
+				.AutoHeight()
+				[
+					ProductionMenuListView.ToSharedRef()
+				]
+		];
+}
+
+TSharedRef<ITableRow> SProductionWizard::OnGenerateMenuRow(TSharedPtr<FProductionWizardMenuEntry> MenuEntry, const TSharedRef<STableViewBase>& OwnerTable)
+{
+	return SNew(STableRow<TSharedPtr<FProductionWizardMenuEntry>>, OwnerTable)
+		.ShowSelection(true)
+		.Padding(FMargin(4.0f, 4.0f, 0.0f, 4.0f))
+		.Content()
+		[
+			SNew(SHorizontalBox)
+
+			+ SHorizontalBox::Slot()
+				.AutoWidth()
+				.Padding(0.0f, 0.0f, 4.0f, 0.0f)
+				[
+					SNew(SImage).Image(MenuEntry->Icon.GetIcon())
+				]
+
+				+ SHorizontalBox::Slot()
+				.AutoWidth()
+				[
+					SNew(STextBlock).Text(MenuEntry->Label)
+				]
+		];
+}
+
+void SProductionWizard::OnUserMenuSelectionChanged(TSharedPtr<FProductionWizardMenuEntry> SelectedEntry, ESelectInfo::Type SelectInfo)
+{
+	if (SelectInfo != ESelectInfo::Direct)
+	{
+		check((SelectedEntry->Panel));
+
+		ProductionMenuListView->ClearSelection();
+		ContentPanelSlot->AttachWidget(SelectedEntry->Panel.ToSharedRef());
+	}
+}
+
+void SProductionWizard::OnProductionMenuSelectionChanged(TSharedPtr<FProductionWizardMenuEntry> SelectedEntry, ESelectInfo::Type SelectInfo)
+{
+	if (SelectInfo != ESelectInfo::Direct)
+	{
+		check((SelectedEntry->Panel));
+
+		UserMenuListView->ClearSelection();
+		ContentPanelSlot->AttachWidget(SelectedEntry->Panel.ToSharedRef());
+	}
+}
+
+TSharedRef<SWidget> SProductionWizard::MakeRevisionControlPanel()
+{
+	// TODO: Fill in this stubbed out panel
+	return SNew(SVerticalBox)
+
+	+ SVerticalBox::Slot()
+		.FillHeight(1.0f)
+		.HAlign(HAlign_Center)
+		[
+			SNew(SHorizontalBox)
+
+			+ SHorizontalBox::Slot()
+				.FillWidth(1.0f)
+				.VAlign(VAlign_Center)
+				[
+					SNew(STextBlock).Text(LOCTEXT("RevisionControlPanel", "REVISION CONTROL PANEL"))
+				]
+		];
+}
+
+TSharedRef<SWidget> SProductionWizard::MakeNamingTokensPanel()
+{
+	// TODO: Fill in this stubbed out panel
+	return SNew(SVerticalBox)
+
+	+ SVerticalBox::Slot()
+		.FillHeight(1.0f)
+		.HAlign(HAlign_Center)
+		[
+			SNew(SHorizontalBox)
+
+			+ SHorizontalBox::Slot()
+				.FillWidth(1.0f)
+				.VAlign(VAlign_Center)
+				[
+					SNew(STextBlock).Text(LOCTEXT("NamingTokensPanel", "NAMING TOKENS PANEL"))
+				]
+		];
+}
+
+TSharedRef<SWidget> SProductionWizard::MakeAssetNamingPanel()
+{
+	// TODO: Fill in this stubbed out panel
+	return SNew(SVerticalBox)
+
+	+ SVerticalBox::Slot()
+		.FillHeight(1.0f)
+		.HAlign(HAlign_Center)
+		[
+			SNew(SHorizontalBox)
+
+			+ SHorizontalBox::Slot()
+				.FillWidth(1.0f)
+				.VAlign(VAlign_Center)
+				[
+					SNew(STextBlock).Text(LOCTEXT("AssetNamingPanel", "ASSET NAMING PANEL"))
+				]
+		];
+}
+
+TSharedRef<SWidget> SProductionWizard::MakeFolderHierarchyPanel()
+{
+	// TODO: Fill in this stubbed out panel
+	return SNew(SVerticalBox)
+
+	+ SVerticalBox::Slot()
+		.FillHeight(1.0f)
+		.HAlign(HAlign_Center)
+		[
+			SNew(SHorizontalBox)
+
+			+ SHorizontalBox::Slot()
+				.FillWidth(1.0f)
+				.VAlign(VAlign_Center)
+				[
+					SNew(STextBlock).Text(LOCTEXT("FolderHierarchyPanel", "FOLDER HIERARCHY PANEL"))
+				]
+		];
+}
+
+#undef LOCTEXT_NAMESPACE
diff --git a/Engine/Plugins/VirtualProduction/Shots/Source/ShotsEditor/Private/UI/SProductionWizard.h b/Engine/Plugins/VirtualProduction/Shots/Source/ShotsEditor/Private/UI/SProductionWizard.h
new file mode 100644
--- /dev/null
+++ b/Engine/Plugins/VirtualProduction/Shots/Source/ShotsEditor/Private/UI/SProductionWizard.h
@@ -0,0 +1,72 @@
+// Copyright Epic Games, Inc. All Rights Reserved.
+
+#pragma once
+
+#include "Widgets/SCompoundWidget.h"
+
+#include "Textures/SlateIcon.h"
+#include "Widgets/DeclarativeSyntaxSupport.h"
+#include "Widgets/Layout/SSplitter.h"
+#include "Widgets/Views/SListView.h"
+
+/** Structure for the list view of menu entires */
+struct FProductionWizardMenuEntry
+{
+	/** Displayed menu entry text */
+	FText Label;
+
+	/** Displayed menu entry icon */
+	FSlateIcon Icon;
+
+	/** Widget to attach to the content panel slot when this menu entry is selected */
+	TSharedPtr<SWidget> Panel;
+};
+
+/**
+ * Main UI widget for the Production Wizard Tool
+ */
+class SProductionWizard : public SCompoundWidget
+{
+public:
+	SProductionWizard() = default;
+
+	SLATE_BEGIN_ARGS(SProductionWizard) {}
+	SLATE_END_ARGS()
+
+	void Construct(const FArguments& InArgs);
+
+private:
+	/** Creates the menu on the left side of the production wizard */
+	TSharedRef<SWidget> MakeMenuPanel();
+
+	/** Creates the buttons on the bottom of the production wizard */
+	TSharedRef<SWidget> MakeButtonsPanel();
+
+	/** Functions for creating the panels for each of the production wizard's menu entries */
+	TSharedRef<SWidget> MakeRevisionControlPanel();
+	TSharedRef<SWidget> MakeNamingTokensPanel();
+	TSharedRef<SWidget> MakeAssetNamingPanel();
+	TSharedRef<SWidget> MakeFolderHierarchyPanel();
+
+	/** Generates a row in the menu list views */
+	TSharedRef<ITableRow> OnGenerateMenuRow(TSharedPtr<FProductionWizardMenuEntry> MenuEntry, const TSharedRef<STableViewBase>& OwnerTable);
+
+	/** Callback when one of the user menu entries is selected, which will cause its panel to be displayed on the right side of the production wizard */
+	void OnUserMenuSelectionChanged(TSharedPtr<FProductionWizardMenuEntry> SelectedEntry, ESelectInfo::Type SelectInfo);
+
+	/** Callback when one of the production menu entries is selected, which will cause its panel to be displayed on the right side of the production wizard */
+	void OnProductionMenuSelectionChanged(TSharedPtr<FProductionWizardMenuEntry> SelectedEntry, ESelectInfo::Type SelectInfo);
+
+private:
+	/** Slots for the menu and menu entry content panels */
+	SSplitter::FSlot* MenuPanelSlot = nullptr;
+	SSplitter::FSlot* ContentPanelSlot = nullptr;
+
+	/** List items sources for the menu list views */
+	TArray<TSharedPtr<FProductionWizardMenuEntry>> UserMenuEntries;
+	TArray<TSharedPtr<FProductionWizardMenuEntry>> ProductionMenuEntries;
+
+	/** Menu list views, which allow the user to select a step in the production wizard */
+	TSharedPtr<SListView<TSharedPtr<FProductionWizardMenuEntry>>> UserMenuListView;
+	TSharedPtr<SListView<TSharedPtr<FProductionWizardMenuEntry>>> ProductionMenuListView;
+};
diff --git a/Engine/Plugins/VirtualProduction/Shots/Source/ShotsEditor/Private/UI/SSequencerSettingsPanel.cpp b/Engine/Plugins/VirtualProduction/Shots/Source/ShotsEditor/Private/UI/SSequencerSettingsPanel.cpp
new file mode 100644
--- /dev/null
+++ b/Engine/Plugins/VirtualProduction/Shots/Source/ShotsEditor/Private/UI/SSequencerSettingsPanel.cpp
@@ -0,0 +1,208 @@
+// Copyright Epic Games, Inc. All Rights Reserved.
+
+#include "SSequencerSettingsPanel.h"
+
+#include "Framework/MultiBox/MultiBoxBuilder.h"
+#include "ProductionSettings.h"
+#include "ShotsEditorStyle.h"
+#include "UI/SActiveProductionCombo.h"
+#include "Widgets/Input/SComboButton.h"
+#include "Widgets/Input/SNumericEntryBox.h"
+#include "Widgets/Layout/SBorder.h"
+#include "Widgets/Layout/SSeparator.h"
+#include "Widgets/SBoxPanel.h"
+#include "Widgets/SFrameRatePicker.h"
+
+#define LOCTEXT_NAMESPACE "SSequencerSettingsPanel"
+
+void SSequencerSettingsPanel::Construct(const FArguments& InArgs)
+{
+	UProductionSettings* ProductionSettings = GetMutableDefault<UProductionSettings>();
+
+	auto IsActiveProductionValid = []() -> bool
+		{
+			const UProductionSettings* const ProductionSettings = GetDefault<UProductionSettings>();
+			return (ProductionSettings->GetActiveProduction() != nullptr);
+		};
+
+	// Get a text representation of the active subsequence priority to display in the combo button below
+	auto GetSubsequencePriorityText = [ProductionSettings]() -> FText
+		{
+			if (ProductionSettings->GetActiveSubsequencePriority() == ESubsequencePriority::TopDown)
+			{
+				return LOCTEXT("TopDownText", "Top Down");
+			}
+			return LOCTEXT("BottomUpText", "Bottom Up");
+		};
+
+	// Build the menu options for the subsequence priority combo button below
+	FMenuBuilder SubsequencePriorityMenuBuilder(true, nullptr);
+	{
+		SubsequencePriorityMenuBuilder.AddMenuEntry(
+			LOCTEXT("TopDownText", "Top Down"),
+			LOCTEXT("TopDownText", "Top Down"),
+			FSlateIcon(),
+			FUIAction(FExecuteAction::CreateLambda([ProductionSettings]() { ProductionSettings->SetActiveSubsequencePriority(ESubsequencePriority::TopDown); })),
+			NAME_None,
+			EUserInterfaceActionType::Button
+		);
+
+		SubsequencePriorityMenuBuilder.AddMenuEntry(
+			LOCTEXT("BottomUpText", "Bottom Up"),
+			LOCTEXT("BottomUpText", "Bottom Up"),
+			FSlateIcon(),
+			FUIAction(FExecuteAction::CreateLambda([ProductionSettings]() { ProductionSettings->SetActiveSubsequencePriority(ESubsequencePriority::BottomUp); })),
+			NAME_None,
+			EUserInterfaceActionType::Button
+		);
+	}
+
+	ChildSlot
+		[
+			SNew(SVerticalBox)
+
+				// Active Production Selector
+				+ SVerticalBox::Slot()
+				.AutoHeight()
+				[
+					SNew(SActiveProductionCombo)
+				]
+
+				// Separator
+				+ SVerticalBox::Slot()
+				.AutoHeight()
+				[
+					SNew(SSeparator)
+						.Orientation(Orient_Horizontal)
+						.Thickness(2.0f)
+				]
+
+				// Sequencer Settings Panel
+				+ SVerticalBox::Slot()
+				.FillHeight(1.0f)
+				[
+					SNew(SBorder)
+						.BorderImage(FAppStyle::Get().GetBrush("Brushes.Panel"))
+						.Padding(16.0f)
+						[
+							SNew(SVerticalBox)
+
+								// Title
+								+ SVerticalBox::Slot()
+								.AutoHeight()
+								.Padding(0.0f, 0.0f, 0.0f, 8.0f)
+								[
+									SNew(STextBlock)
+										.Text(LOCTEXT("SequencerSettingsTitle", "Production Settings"))
+										.Font(FShotsEditorStyle::Get().GetFontStyle("ProductionWizard.TitleFont"))
+								]
+
+								// Heading
+								+ SVerticalBox::Slot()
+								.AutoHeight()
+								.Padding(0.0f, 0.0f, 0.0f, 8.0f)
+								[
+									SNew(STextBlock)
+										.Text(LOCTEXT("SequencerSettingsHeading", "Sequencer Settings"))
+										.Font(FShotsEditorStyle::Get().GetFontStyle("ProductionWizard.HeadingFont"))
+								]
+
+								// Info Text
+								+ SVerticalBox::Slot()
+								.AutoHeight()
+								.Padding(0.0f, 0.0f, 0.0f, 20.0f)
+								[
+									SNew(STextBlock)
+										.Text(LOCTEXT("SequencerSettingsInfoText", "Configure basic settings that apply to all level sequences and shots."))
+								]
+
+								// Frame Rate setting
+								+ SVerticalBox::Slot()
+								.AutoHeight()
+								.Padding(0.0f, 0.0f, 0.0f, 8.0f)
+								[
+									SNew(SHorizontalBox)
+
+										+ SHorizontalBox::Slot()
+										.FillWidth(0.3f)
+										.VAlign(VAlign_Center)
+										[
+											SNew(STextBlock).Text(LOCTEXT("FrameRateText", "Frame Rate"))
+										]
+
+										+ SHorizontalBox::Slot()
+										.FillWidth(0.7f)
+										.VAlign(VAlign_Center)
+										[
+											SNew(SFrameRatePicker)
+												.HasMultipleValues(false)
+												.Value_Lambda([ProductionSettings]() { return ProductionSettings->GetActiveDisplayRate(); })
+												.OnValueChanged_Lambda([ProductionSettings](FFrameRate NewFrameRate) { ProductionSettings->SetActiveDisplayRate(NewFrameRate); })
+												.IsEnabled_Lambda(IsActiveProductionValid)
+										]
+								]
+
+								// Start Frame setting
+								+ SVerticalBox::Slot()
+								.AutoHeight()
+								.Padding(0.0f, 0.0f, 0.0f, 8.0f)
+								[
+									SNew(SHorizontalBox)
+
+										+ SHorizontalBox::Slot()
+										.FillWidth(0.3f)
+										.VAlign(VAlign_Center)
+										[
+											SNew(STextBlock).Text(LOCTEXT("SequencerStartFrameText", "Sequencer Start Frame"))
+										]
+
+										+ SHorizontalBox::Slot()
+										.FillWidth(0.7f)
+										.VAlign(VAlign_Center)
+										[
+											SNew(SNumericEntryBox<int32>)
+												.AllowSpin(false)
+												.Value_Lambda([ProductionSettings]() { return ProductionSettings->GetActiveStartFrame(); })
+												.OnValueCommitted_Lambda([ProductionSettings](int32 NewValue, ETextCommit::Type CommitType) { ProductionSettings->SetActiveStartFrame(NewValue); })
+												.IsEnabled_Lambda(IsActiveProductionValid)
+										]
+								]
+
+								// Subsequence Priority setting
+								+ SVerticalBox::Slot()
+								.AutoHeight()
+								.Padding(0.0f, 0.0f, 0.0f, 8.0f)
+								[
+									SNew(SHorizontalBox)
+
+										+ SHorizontalBox::Slot()
+										.FillWidth(0.3f)
+										.VAlign(VAlign_Center)
+										[
+											SNew(STextBlock).Text(LOCTEXT("SubsequencePriorityText", "Subsequence Priority"))
+										]
+
+										+ SHorizontalBox::Slot()
+										.FillWidth(0.7f)
+										.VAlign(VAlign_Center)
+										[
+											SNew(SComboButton)
+												.IsEnabled_Lambda(IsActiveProductionValid)
+												.MenuContent()
+												[
+													SubsequencePriorityMenuBuilder.MakeWidget()
+												]
+												.ButtonContent()
+												[
+													SNew(STextBlock)
+														.Text_Lambda(GetSubsequencePriorityText)
+														.Font(FAppStyle::Get().GetFontStyle("Normal"))
+												]
+										]
+								]
+						]
+				]
+		];
+}
+
+#undef LOCTEXT_NAMESPACE
diff --git a/Engine/Plugins/VirtualProduction/Shots/Source/ShotsEditor/Private/UI/SSequencerSettingsPanel.h b/Engine/Plugins/VirtualProduction/Shots/Source/ShotsEditor/Private/UI/SSequencerSettingsPanel.h
new file mode 100644
--- /dev/null
+++ b/Engine/Plugins/VirtualProduction/Shots/Source/ShotsEditor/Private/UI/SSequencerSettingsPanel.h
@@ -0,0 +1,19 @@
+// Copyright Epic Games, Inc. All Rights Reserved.
+
+#pragma once
+
+#include "Widgets/SCompoundWidget.h"
+
+#include "Widgets/DeclarativeSyntaxSupport.h"
+
+/**
+ * UI for the Sequencer Settings panel in the Production Wizard
+ */
+class SSequencerSettingsPanel : public SCompoundWidget
+{
+public:
+	SLATE_BEGIN_ARGS(SSequencerSettingsPanel) {}
+	SLATE_END_ARGS()
+
+	void Construct(const FArguments& InArgs);
+};
