ShotManagement: Add Naming Tokens panel to Production Wizard.
Added method to UNamingTokensEditorSubsystem to get a set of all discovered token namespaces (native and BP).
Added a DenyList of namespaces to UProductionSettings to maintain a list of namespaces that should not be evaluated by the UNamingTokensEditorSubsystem when that production is active.

diff --git a/Engine/Plugins/Developer/NamingTokens/Source/NamingTokens/Private/NamingTokens.cpp b/Engine/Plugins/Developer/NamingTokens/Source/NamingTokens/Private/NamingTokens.cpp
--- a/Engine/Plugins/Developer/NamingTokens/Source/NamingTokens/Private/NamingTokens.cpp
+++ b/Engine/Plugins/Developer/NamingTokens/Source/NamingTokens/Private/NamingTokens.cpp
@@ -128,6 +128,15 @@ const TArray<FNamingTokenData>& UNamingTokens::GetCustomTokens() const
 	return CustomTokens;
 }
 
+TArray<FNamingTokenData> UNamingTokens::GetAllTokens() const
+{
+	TArray<FNamingTokenData> AllTokens;
+	AllTokens.Reserve(DefaultTokens.Num() + CustomTokens.Num());
+	AllTokens.Append(DefaultTokens);
+	AllTokens.Append(CustomTokens);
+	return AllTokens;
+}
+
 FDateTime UNamingTokens::GetCurrentDateTime_Implementation() const
 {
 	return CurrentEvaluationData.CurrentDateTime;
diff --git a/Engine/Plugins/Developer/NamingTokens/Source/NamingTokens/Private/NamingTokensEditorSubsystem.cpp b/Engine/Plugins/Developer/NamingTokens/Source/NamingTokens/Private/NamingTokensEditorSubsystem.cpp
--- a/Engine/Plugins/Developer/NamingTokens/Source/NamingTokens/Private/NamingTokensEditorSubsystem.cpp
+++ b/Engine/Plugins/Developer/NamingTokens/Source/NamingTokens/Private/NamingTokensEditorSubsystem.cpp
@@ -13,6 +13,41 @@
 #include "Internationalization/Regex.h"
 #include "Misc/PackageName.h"
 
+TArray<FString> UNamingTokensEditorSubsystem::GetAllNamespaces() const
+{
+	const FAssetRegistryModule& AssetRegistryModule = FModuleManager::LoadModuleChecked<FAssetRegistryModule>("AssetRegistry");
+	const IAssetRegistry& AssetRegistry = AssetRegistryModule.Get();
+
+	// Find all assets whose native parent class is UNamingTokens, and add their namespaces
+	TArray<FAssetData> NamingTokenAssets;
+	const TMultiMap<FName, FString> TagValues = { { FBlueprintTags::NativeParentClassPath, FObjectPropertyBase::GetExportPath(UNamingTokens::StaticClass()) } };
+	AssetRegistryModule.Get().GetAssetsByTagValues(TagValues, NamingTokenAssets);
+
+	TSet<FString> DiscoveredNamespaces;
+	for (const FAssetData& NamingTokenAssetData : NamingTokenAssets)
+	{
+		FString FoundNamespace;
+		NamingTokenAssetData.GetTagValue(UNamingTokens::GetNamespacePropertyName(), FoundNamespace);
+
+		DiscoveredNamespaces.Add(FoundNamespace);
+	}
+
+	// Find all native classes derived from UNamingTokens, and add their namespaces
+	TArray<UClass*> DerivedClasses;
+	GetDerivedClasses(UNamingTokens::StaticClass(), DerivedClasses);
+	for (const UClass* DerivedClass : DerivedClasses)
+	{
+		constexpr EClassFlags InvalidClassFlags = CLASS_Deprecated | CLASS_NewerVersionExists | CLASS_Abstract;
+		if (DerivedClass->IsNative() && !DerivedClass->HasAnyClassFlags(InvalidClassFlags))
+		{
+			FString Namepsace = DerivedClass->GetDefaultObject<UNamingTokens>()->GetNamespace();
+			DiscoveredNamespaces.Add(Namepsace);
+		}
+	}
+
+	return DiscoveredNamespaces.Array();
+}
+
 UNamingTokens* UNamingTokensEditorSubsystem::GetNamingTokens(const FString& InNamespace) const
 {
 	auto LoadNamingToken = [this, &InNamespace](const TSoftClassPtr<UNamingTokens>& InTokensClass) -> UNamingTokens*
diff --git a/Engine/Plugins/Developer/NamingTokens/Source/NamingTokens/Public/NamingTokens.h b/Engine/Plugins/Developer/NamingTokens/Source/NamingTokens/Public/NamingTokens.h
--- a/Engine/Plugins/Developer/NamingTokens/Source/NamingTokens/Public/NamingTokens.h
+++ b/Engine/Plugins/Developer/NamingTokens/Source/NamingTokens/Public/NamingTokens.h
@@ -50,6 +50,9 @@ public:
 	/** Retrieve the custom tokens. */
 	const TArray<FNamingTokenData>& GetCustomTokens() const;
 	
+	/** Retrieve all tokens. */
+	TArray<FNamingTokenData> GetAllTokens() const;
+
 	/** Retrieve the naming token's namespace. */
 	const FString& GetNamespace() const { return Namespace; }
 	/** The name of the namespace property. */
diff --git a/Engine/Plugins/Developer/NamingTokens/Source/NamingTokens/Public/NamingTokensEditorSubsystem.h b/Engine/Plugins/Developer/NamingTokens/Source/NamingTokens/Public/NamingTokensEditorSubsystem.h
--- a/Engine/Plugins/Developer/NamingTokens/Source/NamingTokens/Public/NamingTokensEditorSubsystem.h
+++ b/Engine/Plugins/Developer/NamingTokens/Source/NamingTokens/Public/NamingTokensEditorSubsystem.h
@@ -85,6 +85,10 @@ public:
 	UFUNCTION(BlueprintCallable, Category = "NamingTokens")
 	TArray<FString> GetGlobalNamespaces() const;
 
+	/** Retrieve all discovered namespaces. */
+	UFUNCTION(BlueprintCallable, Category = "NamingTokens")
+	TArray<FString> GetAllNamespaces() const;
+
 	/** Locate all reference naming token namespaces from a given string. */
 	TSet<FString> GetNamingTokenNamespacesFromString(const FString& InTokenString, const FNamingTokenFilterArgs& InFilter = FNamingTokenFilterArgs()) const;
 	
diff --git a/Engine/Plugins/VirtualProduction/Shots/Shots.uplugin b/Engine/Plugins/VirtualProduction/Shots/Shots.uplugin
--- a/Engine/Plugins/VirtualProduction/Shots/Shots.uplugin
+++ b/Engine/Plugins/VirtualProduction/Shots/Shots.uplugin
@@ -32,6 +32,10 @@
 		{
 			"Name": "DirectoryPlaceholder",
 			"Enabled": true
+		},
+		{
+			"Name": "NamingTokens",
+			"Enabled": true
 		}
 	]
 }
\ No newline at end of file
diff --git a/Engine/Plugins/VirtualProduction/Shots/Source/ShotsEditor/Private/ProductionSettings.cpp b/Engine/Plugins/VirtualProduction/Shots/Source/ShotsEditor/Private/ProductionSettings.cpp
--- a/Engine/Plugins/VirtualProduction/Shots/Source/ShotsEditor/Private/ProductionSettings.cpp
+++ b/Engine/Plugins/VirtualProduction/Shots/Source/ShotsEditor/Private/ProductionSettings.cpp
@@ -428,6 +428,26 @@ void UProductionSettings::SetSubsequencePriority(FGuid ProductionID, ESubsequenc
 	}
 }
 
+void UProductionSettings::AddNamespaceToDenyList(FGuid ProductionID, const FString& Namespace)
+{
+	FCinematicProduction* Production = Productions.FindByPredicate([ProductionID](const FCinematicProduction& Other) { return Other.ProductionID == ProductionID; });
+	if (Production)
+	{
+		Production->NamingTokenNamespaceDenyList.Add(Namespace);
+		TryUpdateDefaultConfigFile();
+	}
+}
+
+void UProductionSettings::RemoveNamespaceFromDenyList(FGuid ProductionID, const FString& Namespace)
+{
+	FCinematicProduction* Production = Productions.FindByPredicate([ProductionID](const FCinematicProduction& Other) { return Other.ProductionID == ProductionID; });
+	if (Production)
+	{
+		Production->NamingTokenNamespaceDenyList.Remove(Namespace);
+		TryUpdateDefaultConfigFile();
+	}
+}
+
 void UProductionSettings::SetActiveDisplayRate(FFrameRate DisplayRate)
 {
 	if (ActiveProduction)
diff --git a/Engine/Plugins/VirtualProduction/Shots/Source/ShotsEditor/Private/ProductionSettings.h b/Engine/Plugins/VirtualProduction/Shots/Source/ShotsEditor/Private/ProductionSettings.h
--- a/Engine/Plugins/VirtualProduction/Shots/Source/ShotsEditor/Private/ProductionSettings.h
+++ b/Engine/Plugins/VirtualProduction/Shots/Source/ShotsEditor/Private/ProductionSettings.h
@@ -30,6 +30,11 @@ struct FCinematicProduction
 public:
 	FCinematicProduction();
 
+	/** Unique ID of the production */
+	UPROPERTY(BlueprintReadOnly, Category = "Default", meta = (IgnoreForMemberInitializationTest))
+	FGuid ProductionID;
+
+	/** Production Name */
 	UPROPERTY(config, EditAnywhere, BlueprintReadOnly, Category = "Default")
 	FString ProductionName;
 
@@ -45,9 +50,9 @@ public:
 	UPROPERTY(config, EditAnywhere, BlueprintReadOnly, Category = "Default")
 	ESubsequencePriority SubsequencePriority = ESubsequencePriority::BottomUp;
 
-	/** Unique ID of the production */
-	UPROPERTY(meta=(IgnoreForMemberInitializationTest))
-	FGuid ProductionID;
+	/** List of Naming Token namespaces that should not be evaluated */
+	UPROPERTY(config, EditAnywhere, BlueprintReadOnly, Category = "Default")
+	TSet<FString> NamingTokenNamespaceDenyList;
 };
 
 /**
@@ -111,6 +116,12 @@ public:
 	/** Sets the SubsequencePriority of the production matching the input ID */
 	void SetSubsequencePriority(FGuid ProductionID, ESubsequencePriority Priority);
 
+	/** Adds a Naming Token namespace to the DenyList of the production matching the input ID */
+	void AddNamespaceToDenyList(FGuid ProductionID, const FString& Namespace);
+
+	/** Removes a Naming Token namespace from the DenyList of the production matching the input ID */
+	void RemoveNamespaceFromDenyList(FGuid ProductionID, const FString& Namespace);
+
 	/** Sets the DefaultDisplayRate of the active production */
 	void SetActiveDisplayRate(FFrameRate DisplayRate);
 
diff --git a/Engine/Plugins/VirtualProduction/Shots/Source/ShotsEditor/Private/ShotsEditorStyle.cpp b/Engine/Plugins/VirtualProduction/Shots/Source/ShotsEditor/Private/ShotsEditorStyle.cpp
--- a/Engine/Plugins/VirtualProduction/Shots/Source/ShotsEditor/Private/ShotsEditorStyle.cpp
+++ b/Engine/Plugins/VirtualProduction/Shots/Source/ShotsEditor/Private/ShotsEditorStyle.cpp
@@ -35,7 +35,7 @@ FShotsEditorStyle::FShotsEditorStyle()
 	Set("Icons.Export", new FSlateVectorImageBrush(RootToContentDir(TEXT("Slate/ExportAll.svg")), Icon16x16));
 
 	// Fonts
-	Set("ProductionWizard.TitleFont", DEFAULT_FONT("Normal", 12));
+	Set("ProductionWizard.TitleFont", DEFAULT_FONT("Bold", 10));
 	Set("ProductionWizard.HeadingFont", DEFAULT_FONT("Bold", 14));
 
 	// Brushes
diff --git a/Engine/Plugins/VirtualProduction/Shots/Source/ShotsEditor/Private/UI/SNamingTokensPanel.cpp b/Engine/Plugins/VirtualProduction/Shots/Source/ShotsEditor/Private/UI/SNamingTokensPanel.cpp
new file mode 100644
--- /dev/null
+++ b/Engine/Plugins/VirtualProduction/Shots/Source/ShotsEditor/Private/UI/SNamingTokensPanel.cpp
@@ -0,0 +1,412 @@
+// Copyright Epic Games, Inc. All Rights Reserved.
+
+#include "SNamingTokensPanel.h"
+
+#include "Algo/Reverse.h"
+#include "Editor.h"
+#include "GlobalNamingTokens.h"
+#include "NamingTokenData.h"
+#include "NamingTokensEditorSubsystem.h"
+#include "ProductionSettings.h"
+#include "ShotsEditorStyle.h"
+#include "UI/SActiveProductionCombo.h"
+#include "Widgets/SBoxPanel.h"
+#include "Widgets/Input/SCheckBox.h"
+#include "Widgets/Layout/SBorder.h"
+#include "Widgets/Layout/SSeparator.h"
+
+#define LOCTEXT_NAMESPACE "SNamingTokensPanel"
+
+void SNamingTokensPanel::Construct(const FArguments& InArgs)
+{
+	const UNamingTokensEditorSubsystem* const NamingTokensSubsystem = GEditor->GetEditorSubsystem<UNamingTokensEditorSubsystem>();
+
+	// Build the list of the global tokens
+	const UNamingTokens* const GlobalTokens = NamingTokensSubsystem->GetNamingTokens(UGlobalNamingTokens::GetGlobalNamespace());
+	const TArray<FNamingTokenData>& GlobalTokenData = GlobalTokens->GetDefaultTokens();
+	for (const FNamingTokenData& TokenData : GlobalTokenData)
+	{
+		GlobalTokenListItems.Add(MakeShared<FNamingTokenData>(TokenData));
+	}
+	GlobalTokenListItems.Sort([](const TSharedPtr<FNamingTokenData>& A, const TSharedPtr<FNamingTokenData>& B) { return A->TokenKey < B->TokenKey; });
+
+	// Build the list of all discovered naming token namespaces
+	TArray<FString> TokenNamespaces = NamingTokensSubsystem->GetAllNamespaces();
+
+	// The global namespace is displayed separately, and also cannot be added to the active production settings DenyList
+	TokenNamespaces.Remove(UGlobalNamingTokens::GetGlobalNamespace());
+
+	for (const FString& GlobalNamespace : TokenNamespaces)
+	{
+		NamingTokenNamespaceListItems.Add(MakeShared<FString>(GlobalNamespace));
+	}
+	NamingTokenNamespaceListItems.Sort([](const TSharedPtr<FString>& A, const TSharedPtr<FString>& B) { return *A < *B; });
+
+	// The custom token list starts empty and is populated when one of the namespace list items is selected
+	CustomTokenListItems.Empty();
+
+	ChildSlot
+		[
+			SNew(SVerticalBox)
+
+			// Active Production Selector
+			+ SVerticalBox::Slot()
+				.AutoHeight()
+				[
+					SNew(SActiveProductionCombo)
+				]
+
+			// Separator
+			+ SVerticalBox::Slot()
+				.AutoHeight()
+				[
+					SNew(SSeparator)
+						.Orientation(Orient_Horizontal)
+						.Thickness(2.0f)
+				]
+
+			// Naming Tokens Panel
+			+ SVerticalBox::Slot()
+				.FillHeight(1.0f)
+				[
+					SNew(SBorder)
+						.BorderImage(FAppStyle::Get().GetBrush("Brushes.Panel"))
+						.Padding(16.0f)
+						[
+							SNew(SVerticalBox)
+
+							// Title
+							+ SVerticalBox::Slot()
+								.AutoHeight()
+								.Padding(0.0f, 0.0f, 0.0f, 4.0f)
+								[
+									SNew(STextBlock)
+										.Text(LOCTEXT("NamingTokensTitle", "Production Settings"))
+										.Font(FShotsEditorStyle::Get().GetFontStyle("ProductionWizard.TitleFont"))
+								]
+
+							// Heading
+							+ SVerticalBox::Slot()
+								.AutoHeight()
+								.Padding(0.0f, 0.0f, 0.0f, 4.0f)
+								[
+									SNew(STextBlock)
+										.Text(LOCTEXT("NamingTokensHeading", "Naming Tokens"))
+										.Font(FShotsEditorStyle::Get().GetFontStyle("ProductionWizard.HeadingFont"))
+								]
+
+							// Info Text
+							+ SVerticalBox::Slot()
+								.AutoHeight()
+								.Padding(0.0f, 0.0f, 0.0f, 8.0f)
+								[
+									SNew(STextBlock)
+										.Text(LOCTEXT("NamingTokensInfoText", "Naming tokens insert relevant information automatically into your asset names."))
+								]
+
+							// Global Tokens Info Text
+							+ SVerticalBox::Slot()
+								.AutoHeight()
+								.Padding(0.0f, 0.0f, 0.0f, 16.0f)
+								[
+									SNew(STextBlock)
+										.Text(LOCTEXT("GlobalTokensInfoText", "The following global tokens are built-in tokens you can use in all projects and tools. They are always replaced by the values shown here."))
+										.AutoWrapText(true)
+								]
+
+							// Global Tokens List View
+							+ SVerticalBox::Slot()
+								.AutoHeight()
+								.MaxHeight(26 + 24 * 5)
+								.Padding(0.0f, 0.0f, 0.0f, 16.0f)
+								[
+									SAssignNew(GlobalTokenListView, SListView<TSharedPtr<FNamingTokenData>>)
+										.ListItemsSource(&GlobalTokenListItems)
+										.OnGenerateRow(this, &SNamingTokensPanel::OnGenerateNamingTokenRow)
+										.SelectionMode(ESelectionMode::None)
+										.HeaderRow
+										(
+											SNew(SHeaderRow)
+
+											+ SHeaderRow::Column("NamingTokens")
+												.DefaultLabel(LOCTEXT("GlobalTokensLabel", "Global Tokens"))
+												.FillWidth(0.3f)
+												.OnSort(this, &SNamingTokensPanel::HandleGlobalTokenListSort)
+												.SortMode_Lambda([this]() { return GlobalTokenListSortMode; })
+
+											+ SHeaderRow::Column("Description")
+												.DefaultLabel(LOCTEXT("DescriptionLabel", "Description"))
+												.FillWidth(0.7f)
+										)
+								]
+
+							// Separator
+							+ SVerticalBox::Slot()
+								.AutoHeight()
+								.Padding(0.0f, 0.0f, 0.0f, 16.0f)
+								[
+									SNew(SSeparator)
+										.Orientation(Orient_Horizontal)
+										.Thickness(2.0f)
+								]
+
+							// Custom Tokens Info Text
+							+ SVerticalBox::Slot()
+								.AutoHeight()
+								.Padding(0.0f, 0.0f, 0.0f, 16.0f)
+								[
+									SNew(STextBlock)
+										.Text(LOCTEXT("CustomTokensInfoText", "Your project also defines the following namespaces and custom naming tokens."))
+										.AutoWrapText(true)
+								]
+
+							// Custom Naming Tokens
+							+ SVerticalBox::Slot()
+								.FillHeight(1.0f)
+								.Padding(0.0f, 0.0f, 0.0f, 0.0f)
+								[
+									SNew(SBorder)
+										.BorderImage(FAppStyle::Get().GetBrush("Brushes.Background"))
+										.Padding(4.0f)
+										[
+											SNew(SHorizontalBox)
+
+											// Namespaces List View
+											+ SHorizontalBox::Slot()
+												.FillWidth(0.5f)
+												.Padding(0.0f, 0.0f, 4.0f, 0.0f)
+												[
+													SNew(SOverlay)
+
+													+ SOverlay::Slot()
+														[
+															SAssignNew(NamingTokenNamespaceListView, SListView<TSharedPtr<FString>>)
+																.ListItemsSource(&NamingTokenNamespaceListItems)
+																.OnGenerateRow(this, &SNamingTokensPanel::OnGenerateNamingTokenNamespaceRow)
+																.SelectionMode(ESelectionMode::Single)
+																.OnSelectionChanged(this, &SNamingTokensPanel::OnNamespaceSelectionChanged)
+																.HeaderRow
+																(
+																	SNew(SHeaderRow)
+
+																	+ SHeaderRow::Column("Namespaces")
+																		.DefaultLabel(LOCTEXT("NamespacesLabel", "Namespaces"))
+																		.OnSort(this, &SNamingTokensPanel::HandleNamespaceListSort)
+																		.SortMode_Lambda([this]() { return NamespaceListSortMode; })
+																)
+														]
+
+													+ SOverlay::Slot()
+														.Padding(0.0f, 42.0f, 0.0f, 0.0f)
+														.HAlign(HAlign_Center)
+														[
+															SNew(STextBlock)
+																.Text(LOCTEXT("NoNamespacesFoundHintText", "No namespaces found in your project."))
+																.TextStyle(FAppStyle::Get(), "HintText")
+																.Visibility_Lambda([this]() -> EVisibility
+																	{
+																		return NamingTokenNamespaceListItems.IsEmpty() ? EVisibility::Visible : EVisibility::Collapsed;
+																	})
+														]
+												]
+
+											// Custom Naming Token List View
+											+ SHorizontalBox::Slot()
+												.FillWidth(0.5f)
+												[
+													SNew(SOverlay)
+
+													+ SOverlay::Slot()
+														[
+															SAssignNew(CustomTokenListView, SListView<TSharedPtr<FNamingTokenData>>)
+																.ListItemsSource(&CustomTokenListItems)
+																.OnGenerateRow(this, &SNamingTokensPanel::OnGenerateNamingTokenRow)
+																.SelectionMode(ESelectionMode::None)
+																.HeaderRow
+																(
+																	SNew(SHeaderRow)
+
+																	+ SHeaderRow::Column("NamingTokens")
+																		.DefaultLabel(LOCTEXT("CustomTokensLabel", "Naming Tokens"))
+																		.FillWidth(0.5f)
+																		.OnSort(this, &SNamingTokensPanel::HandleCustomTokenListSort)
+																		.SortMode_Lambda([this]() { return CustomTokenListSortMode; })
+
+																	+ SHeaderRow::Column("Description")
+																		.DefaultLabel(LOCTEXT("DescriptionLabel", "Description"))
+																		.FillWidth(0.5f)
+																)
+														]
+
+													+ SOverlay::Slot()
+														.Padding(0.0f, 42.0f, 0.0f, 0.0f)
+														.HAlign(HAlign_Center)
+														[
+															SNew(STextBlock)
+																.Text(LOCTEXT("NoNamespaceSelectedHintText", "Select a namespace to see its naming tokens."))
+																.TextStyle(FAppStyle::Get(), "HintText")
+																.Visibility_Lambda([this]() -> EVisibility
+																	{
+																		return (NamingTokenNamespaceListView->GetNumItemsSelected() == 0) ? EVisibility::Visible : EVisibility::Collapsed;
+																	})
+														]
+												]
+										]
+								]
+						]
+				]
+		];
+}
+
+void SNamingTokensPanel::OnNamespaceSelectionChanged(TSharedPtr<FString> SelectedNamespace, ESelectInfo::Type SelectInfo)
+{
+	CustomTokenListItems.Empty();
+
+	if (SelectedNamespace)
+	{
+		UNamingTokensEditorSubsystem* NamingTokensSubsystem = GEditor->GetEditorSubsystem<UNamingTokensEditorSubsystem>();
+		if (UNamingTokens* CustomTokens = NamingTokensSubsystem->GetNamingTokens(*SelectedNamespace))
+		{
+			const TArray<FNamingTokenData>& CustomTokenData = CustomTokens->GetAllTokens();
+
+			for (const FNamingTokenData& TokenData : CustomTokenData)
+			{
+				CustomTokenListItems.Add(MakeShared<FNamingTokenData>(TokenData));
+			}
+		}
+
+		if (CustomTokenListSortMode == EColumnSortMode::Ascending)
+		{
+			CustomTokenListItems.Sort([](const TSharedPtr<FNamingTokenData>& A, const TSharedPtr<FNamingTokenData>& B) { return A->TokenKey < B->TokenKey; });
+		}
+		else
+		{
+			CustomTokenListItems.Sort([](const TSharedPtr<FNamingTokenData>& A, const TSharedPtr<FNamingTokenData>& B) { return A->TokenKey > B->TokenKey; });
+		}
+	}
+
+	CustomTokenListView->RebuildList();
+}
+
+TSharedRef<ITableRow> SNamingTokensPanel::OnGenerateNamingTokenRow(TSharedPtr<FNamingTokenData> InTokenData, const TSharedRef<STableViewBase>& OwnerTable)
+{
+	return SNew(SNamingTokenRow, OwnerTable, InTokenData);
+}
+
+TSharedRef<ITableRow> SNamingTokensPanel::OnGenerateNamingTokenNamespaceRow(TSharedPtr<FString> InNamespace, const TSharedRef<STableViewBase>& OwnerTable)
+{
+	auto IsActiveProductionValid = []() -> bool
+		{
+			const UProductionSettings* const ProductionSettings = GetDefault<UProductionSettings>();
+			return (ProductionSettings->GetActiveProduction() != nullptr);
+		};
+
+	return SNew(STableRow<TSharedPtr<FString>>, OwnerTable)
+		.ShowSelection(true)
+		.Padding(FMargin(8.0f, 4.0f, 8.0f, 4.0f))
+		.Content()
+		[
+			SNew(SHorizontalBox)
+
+				+ SHorizontalBox::Slot()
+				.AutoWidth()
+				[
+					SNew(SCheckBox)
+						.Padding(FMargin(0.0f, 0.0f, 4.0f, 0.0f))
+						.OnCheckStateChanged(this, &SNamingTokensPanel::OnNamespaceChecked, InNamespace)
+						.IsChecked(this, &SNamingTokensPanel::GetNamespaceCheckBoxState, InNamespace)
+						.IsEnabled_Lambda(IsActiveProductionValid)
+				]
+
+				+ SHorizontalBox::Slot()
+				.FillWidth(1.0f)
+				.VAlign(VAlign_Center)
+				[
+					SNew(STextBlock)
+						.Text(FText::FromString(*InNamespace))
+				]
+		];
+}
+
+void SNamingTokensPanel::OnNamespaceChecked(ECheckBoxState CheckBoxState, TSharedPtr<FString> Namespace)
+{
+	UProductionSettings* ProductionSettings = GetMutableDefault<UProductionSettings>();
+	if (FCinematicProduction* ActiveProduction = ProductionSettings->GetActiveProduction())
+	{
+		if (CheckBoxState == ECheckBoxState::Checked)
+		{
+			ProductionSettings->RemoveNamespaceFromDenyList(ActiveProduction->ProductionID, *Namespace);
+		}
+		else if (CheckBoxState == ECheckBoxState::Unchecked)
+		{
+			ProductionSettings->AddNamespaceToDenyList(ActiveProduction->ProductionID, *Namespace);
+		}
+	}
+}
+
+ECheckBoxState SNamingTokensPanel::GetNamespaceCheckBoxState(TSharedPtr<FString> Namespace) const
+{
+	const UProductionSettings* const ProductionSettings = GetDefault<UProductionSettings>();
+	if (FCinematicProduction* ActiveProduction = ProductionSettings->GetActiveProduction())
+	{
+		if (ActiveProduction->NamingTokenNamespaceDenyList.Contains(*Namespace))
+		{
+			return ECheckBoxState::Unchecked;
+		}
+	}
+	return ECheckBoxState::Checked;
+}
+
+void SNamingTokensPanel::HandleGlobalTokenListSort(EColumnSortPriority::Type, const FName& ColumnId, EColumnSortMode::Type SortMode)
+{
+	GlobalTokenListSortMode = SortMode;
+	Algo::Reverse(GlobalTokenListItems);
+	GlobalTokenListView->RebuildList();
+}
+
+void SNamingTokensPanel::HandleNamespaceListSort(EColumnSortPriority::Type, const FName& ColumnId, EColumnSortMode::Type SortMode)
+{
+	NamespaceListSortMode = SortMode;
+	Algo::Reverse(NamingTokenNamespaceListItems);
+	NamingTokenNamespaceListView->RebuildList();
+}
+
+void SNamingTokensPanel::HandleCustomTokenListSort(EColumnSortPriority::Type, const FName& ColumnId, EColumnSortMode::Type SortMode)
+{
+	CustomTokenListSortMode = SortMode;
+	Algo::Reverse(CustomTokenListItems);
+	CustomTokenListView->RebuildList();
+}
+
+void SNamingTokenRow::Construct(const FArguments& Args, const TSharedRef<STableViewBase>& OwnerTableView, TSharedPtr<FNamingTokenData>& InRowData)
+{
+	TokenData = InRowData;
+
+	FSuperRowType::FArguments StyleArguments = FSuperRowType::FArguments()
+		.Padding(FMargin(8.0f, 4.0f, 8.0f, 4.0f));
+
+	SMultiColumnTableRow<TSharedPtr<FNamingTokenData>>::Construct(StyleArguments, OwnerTableView);
+}
+
+TSharedRef<SWidget> SNamingTokenRow::GenerateWidgetForColumn(const FName& ColumnName)
+{
+	if (ColumnName == TEXT("NamingTokens"))
+	{
+		const FString FormattedTokenKey = TEXT("{") + TokenData->TokenKey + TEXT("}");
+
+		return SNew(STextBlock)
+			.Text(FText::FromString(FormattedTokenKey));
+	}
+	else if (ColumnName == TEXT("Description"))
+	{
+		const FText ToolTipText = TokenData->Description.IsEmpty() ? TokenData->DisplayName : TokenData->Description;
+
+		return SNew(STextBlock)
+			.Text(TokenData->DisplayName)
+			.ToolTipText(ToolTipText);
+	}
+
+	return SNullWidget::NullWidget;
+}
+
+#undef LOCTEXT_NAMESPACE
diff --git a/Engine/Plugins/VirtualProduction/Shots/Source/ShotsEditor/Private/UI/SNamingTokensPanel.h b/Engine/Plugins/VirtualProduction/Shots/Source/ShotsEditor/Private/UI/SNamingTokensPanel.h
new file mode 100644
--- /dev/null
+++ b/Engine/Plugins/VirtualProduction/Shots/Source/ShotsEditor/Private/UI/SNamingTokensPanel.h
@@ -0,0 +1,76 @@
+// Copyright Epic Games, Inc. All Rights Reserved.
+
+#pragma once
+
+#include "Widgets/SCompoundWidget.h"
+
+#include "Widgets/DeclarativeSyntaxSupport.h"
+#include "Widgets/Views/SListView.h"
+
+struct FNamingTokenData;
+
+/** Row widget for the global tokens list view */
+class SNamingTokenRow : public SMultiColumnTableRow<TSharedPtr<FNamingTokenData>>
+{
+public:
+	SLATE_BEGIN_ARGS(SNamingTokenRow) {}
+	SLATE_END_ARGS()
+
+	void Construct(const FArguments& Args, const TSharedRef<STableViewBase>& OwnerTableView, TSharedPtr<FNamingTokenData>& InRowData);
+
+	/** Creates the widget for this row for the specified column */
+	virtual TSharedRef<SWidget> GenerateWidgetForColumn(const FName& ColumnName) override;
+
+private:
+	/** The underlying token data, used to propertly display its token key and display name */
+	TSharedPtr<FNamingTokenData> TokenData;
+};
+
+/**
+ * UI for the Naming Tokens panel in the Production Wizard
+ */
+class SNamingTokensPanel : public SCompoundWidget
+{
+public:
+	SLATE_BEGIN_ARGS(SNamingTokensPanel) {}
+	SLATE_END_ARGS()
+
+	void Construct(const FArguments& InArgs);
+
+private:
+	/** Generates a row displaying token data */
+	TSharedRef<ITableRow> OnGenerateNamingTokenRow(TSharedPtr<FNamingTokenData> InTokenData, const TSharedRef<STableViewBase>& OwnerTable);
+
+	/** Generates a row displaying a namespace, and a checkbox to add/remove it from the active production settings DenyList */
+	TSharedRef<ITableRow> OnGenerateNamingTokenNamespaceRow(TSharedPtr<FString> InNamespace, const TSharedRef<STableViewBase>& OwnerTable);
+
+	/** Callback when the selection in the namespace list view changes, which updates the custom token list view to display the tokens in the selected namespace */
+	void OnNamespaceSelectionChanged(TSharedPtr<FString> SelectedNamespace, ESelectInfo::Type SelectInfo);
+
+	/** Adds / Removes the associated namespace from the active production settings DenyList */
+	void OnNamespaceChecked(ECheckBoxState CheckBoxState, TSharedPtr<FString> Namespace);
+
+	/** Indicates whether the associated namespace is in the active production settings DenyList */
+	ECheckBoxState GetNamespaceCheckBoxState(TSharedPtr<FString> Namespace) const;
+
+	/** Handles sorting of the items in each list view  in this panel */
+	void HandleGlobalTokenListSort(EColumnSortPriority::Type, const FName& ColumnId, EColumnSortMode::Type SortMode);
+	void HandleNamespaceListSort(EColumnSortPriority::Type, const FName& ColumnId, EColumnSortMode::Type SortMode);
+	void HandleCustomTokenListSort(EColumnSortPriority::Type, const FName& ColumnId, EColumnSortMode::Type SortMode);
+
+private:
+	/** Sources for the list views in this panel */
+	TArray<TSharedPtr<FNamingTokenData>> GlobalTokenListItems;
+	TArray<TSharedPtr<FString>> NamingTokenNamespaceListItems;
+	TArray<TSharedPtr<FNamingTokenData>> CustomTokenListItems;
+
+	/** List view widgets in this panel */
+	TSharedPtr<SListView<TSharedPtr<FNamingTokenData>>> GlobalTokenListView;
+	TSharedPtr<SListView<TSharedPtr<FString>>> NamingTokenNamespaceListView;
+	TSharedPtr<SListView<TSharedPtr<FNamingTokenData>>> CustomTokenListView;
+
+	/** Sort mode for each of the sortable lists in this panel */
+	EColumnSortMode::Type GlobalTokenListSortMode = EColumnSortMode::Ascending;
+	EColumnSortMode::Type NamespaceListSortMode = EColumnSortMode::Ascending;
+	EColumnSortMode::Type CustomTokenListSortMode = EColumnSortMode::Ascending;
+};
diff --git a/Engine/Plugins/VirtualProduction/Shots/Source/ShotsEditor/Private/UI/SProductionListPanel.cpp b/Engine/Plugins/VirtualProduction/Shots/Source/ShotsEditor/Private/UI/SProductionListPanel.cpp
--- a/Engine/Plugins/VirtualProduction/Shots/Source/ShotsEditor/Private/UI/SProductionListPanel.cpp
+++ b/Engine/Plugins/VirtualProduction/Shots/Source/ShotsEditor/Private/UI/SProductionListPanel.cpp
@@ -43,7 +43,7 @@ void SProductionListPanel::Construct(const FArguments& InArgs)
 						// Title
 						+ SVerticalBox::Slot()
 						.AutoHeight()
-						.Padding(0.0f, 0.0f, 0.0f, 8.0f)
+						.Padding(0.0f, 0.0f, 0.0f, 4.0f)
 						[
 							SNew(STextBlock)
 								.Text(LOCTEXT("ProductionsTitle", "User Setup"))
@@ -53,7 +53,7 @@ void SProductionListPanel::Construct(const FArguments& InArgs)
 						// Heading
 						+ SVerticalBox::Slot()
 						.AutoHeight()
-						.Padding(0.0f, 0.0f, 0.0f, 8.0f)
+						.Padding(0.0f, 0.0f, 0.0f, 4.0f)
 						[
 							SNew(STextBlock)
 								.Text(LOCTEXT("ProductionsHeading", "Productions"))
@@ -63,7 +63,7 @@ void SProductionListPanel::Construct(const FArguments& InArgs)
 						// Info Text
 						+ SVerticalBox::Slot()
 						.AutoHeight()
-						.Padding(0.0f, 0.0f, 0.0f, 20.0f)
+						.Padding(0.0f, 0.0f, 0.0f, 16.0f)
 						[
 							SNew(STextBlock)
 								.Text(LOCTEXT("ProductionsInfoText", "You can use this Unreal project for multiple shows, and assign different settings for each show.\n"
diff --git a/Engine/Plugins/VirtualProduction/Shots/Source/ShotsEditor/Private/UI/SProductionWizard.cpp b/Engine/Plugins/VirtualProduction/Shots/Source/ShotsEditor/Private/UI/SProductionWizard.cpp
--- a/Engine/Plugins/VirtualProduction/Shots/Source/ShotsEditor/Private/UI/SProductionWizard.cpp
+++ b/Engine/Plugins/VirtualProduction/Shots/Source/ShotsEditor/Private/UI/SProductionWizard.cpp
@@ -4,6 +4,7 @@
 
 #include "ProductionSettings.h"
 #include "ShotsEditorStyle.h"
+#include "UI/SNamingTokensPanel.h"
 #include "UI/SProductionListPanel.h"
 #include "UI/SSequencerSettingsPanel.h"
 #include "Widgets/Layout/SBorder.h"
@@ -117,7 +118,7 @@ TSharedRef<SWidget> SProductionWizard::MakeMenuPanel()
 	TSharedPtr<FProductionWizardMenuEntry> NamingTokensMenuEntry = MakeShared<FProductionWizardMenuEntry>();
 	NamingTokensMenuEntry->Label = LOCTEXT("NamingTokensMenuLabel", "Naming Tokens");
 	NamingTokensMenuEntry->Icon = FSlateIcon(FShotsEditorStyle::StyleName, "Icons.NamingTokens");
-	NamingTokensMenuEntry->Panel = MakeNamingTokensPanel();
+	NamingTokensMenuEntry->Panel = SNew(SNamingTokensPanel);
 
 	// Build the Asset Naming menu entry, including its content panel
 	TSharedPtr<FProductionWizardMenuEntry> AssetNamingMenuEntry = MakeShared<FProductionWizardMenuEntry>();
@@ -257,26 +258,6 @@ TSharedRef<SWidget> SProductionWizard::MakeRevisionControlPanel()
 		];
 }
 
-TSharedRef<SWidget> SProductionWizard::MakeNamingTokensPanel()
-{
-	// TODO: Fill in this stubbed out panel
-	return SNew(SVerticalBox)
-
-	+ SVerticalBox::Slot()
-		.FillHeight(1.0f)
-		.HAlign(HAlign_Center)
-		[
-			SNew(SHorizontalBox)
-
-			+ SHorizontalBox::Slot()
-				.FillWidth(1.0f)
-				.VAlign(VAlign_Center)
-				[
-					SNew(STextBlock).Text(LOCTEXT("NamingTokensPanel", "NAMING TOKENS PANEL"))
-				]
-		];
-}
-
 TSharedRef<SWidget> SProductionWizard::MakeAssetNamingPanel()
 {
 	// TODO: Fill in this stubbed out panel
diff --git a/Engine/Plugins/VirtualProduction/Shots/Source/ShotsEditor/Private/UI/SProductionWizard.h b/Engine/Plugins/VirtualProduction/Shots/Source/ShotsEditor/Private/UI/SProductionWizard.h
--- a/Engine/Plugins/VirtualProduction/Shots/Source/ShotsEditor/Private/UI/SProductionWizard.h
+++ b/Engine/Plugins/VirtualProduction/Shots/Source/ShotsEditor/Private/UI/SProductionWizard.h
@@ -44,7 +44,6 @@ private:
 
 	/** Functions for creating the panels for each of the production wizard's menu entries */
 	TSharedRef<SWidget> MakeRevisionControlPanel();
-	TSharedRef<SWidget> MakeNamingTokensPanel();
 	TSharedRef<SWidget> MakeAssetNamingPanel();
 	TSharedRef<SWidget> MakeFolderHierarchyPanel();
 
diff --git a/Engine/Plugins/VirtualProduction/Shots/Source/ShotsEditor/Private/UI/SSequencerSettingsPanel.cpp b/Engine/Plugins/VirtualProduction/Shots/Source/ShotsEditor/Private/UI/SSequencerSettingsPanel.cpp
--- a/Engine/Plugins/VirtualProduction/Shots/Source/ShotsEditor/Private/UI/SSequencerSettingsPanel.cpp
+++ b/Engine/Plugins/VirtualProduction/Shots/Source/ShotsEditor/Private/UI/SSequencerSettingsPanel.cpp
@@ -90,7 +90,7 @@ void SSequencerSettingsPanel::Construct(const FArguments& InArgs)
 								// Title
 								+ SVerticalBox::Slot()
 								.AutoHeight()
-								.Padding(0.0f, 0.0f, 0.0f, 8.0f)
+								.Padding(0.0f, 0.0f, 0.0f, 4.0f)
 								[
 									SNew(STextBlock)
 										.Text(LOCTEXT("SequencerSettingsTitle", "Production Settings"))
@@ -100,7 +100,7 @@ void SSequencerSettingsPanel::Construct(const FArguments& InArgs)
 								// Heading
 								+ SVerticalBox::Slot()
 								.AutoHeight()
-								.Padding(0.0f, 0.0f, 0.0f, 8.0f)
+								.Padding(0.0f, 0.0f, 0.0f, 4.0f)
 								[
 									SNew(STextBlock)
 										.Text(LOCTEXT("SequencerSettingsHeading", "Sequencer Settings"))
@@ -110,7 +110,7 @@ void SSequencerSettingsPanel::Construct(const FArguments& InArgs)
 								// Info Text
 								+ SVerticalBox::Slot()
 								.AutoHeight()
-								.Padding(0.0f, 0.0f, 0.0f, 20.0f)
+								.Padding(0.0f, 0.0f, 0.0f, 16.0f)
 								[
 									SNew(STextBlock)
 										.Text(LOCTEXT("SequencerSettingsInfoText", "Configure basic settings that apply to all level sequences and shots."))
diff --git a/Engine/Plugins/VirtualProduction/Shots/Source/ShotsEditor/ShotsEditor.build.cs b/Engine/Plugins/VirtualProduction/Shots/Source/ShotsEditor/ShotsEditor.build.cs
--- a/Engine/Plugins/VirtualProduction/Shots/Source/ShotsEditor/ShotsEditor.build.cs
+++ b/Engine/Plugins/VirtualProduction/Shots/Source/ShotsEditor/ShotsEditor.build.cs
@@ -24,11 +24,13 @@ namespace UnrealBuildTool.Rules
 					"LevelSequenceEditor",
 					"MovieScene",
 					"MovieSceneTools",
+					"NamingTokens",
 					"Projects",
 					"Slate",
 					"SlateCore",
 					"SourceControl",
 					"TimeManagement",
+					"UnrealEd",
 					"WorkspaceMenuStructure"
 				}
 			);
