ShotManagement: Add map of asset classes to default names to the production settings, allowing users of the editor to set their own default names to use when creating assets of various types.
Add Asset Naming panel to Production Wizard.
Fix the active production combo widget to display "None" when no active production is set.

diff --git a/Engine/Plugins/VirtualProduction/Shots/Source/ShotsEditor/Private/ProductionSettings.cpp b/Engine/Plugins/VirtualProduction/Shots/Source/ShotsEditor/Private/ProductionSettings.cpp
--- a/Engine/Plugins/VirtualProduction/Shots/Source/ShotsEditor/Private/ProductionSettings.cpp
+++ b/Engine/Plugins/VirtualProduction/Shots/Source/ShotsEditor/Private/ProductionSettings.cpp
@@ -1,6 +1,7 @@
 // Copyright Epic Games, Inc. All Rights Reserved.
 #include "ProductionSettings.h"
 
+#include "AssetToolsModule.h"
 #include "HAL/FileManager.h"
 #include "HAL/PlatformFileManager.h"
 #include "LevelSequenceProjectSettings.h"
@@ -86,13 +87,13 @@ void UProductionSettings::PostInitProperties()
 	// Apply the initial overrides to existing project settings based on the current active production
 	if (ActiveProduction)
 	{
-		ApplySequencerSettingOverrides();
+		ApplyProjectOverrides();
 	}
 }
 
-void UProductionSettings::ApplySequencerSettingOverrides()
+void UProductionSettings::ApplyProjectOverrides()
 {
-	// Modify the property flags and tooltip of these overridden sequencer settings
+	// Modify the property flags and tooltip of these overridden project settings
 	// They should be read-only when driven by an active production, and writable when the active production is None
 	ULevelSequenceProjectSettings* LevelSequenceSettings = GetMutableDefault<ULevelSequenceProjectSettings>();
 	FProperty* DefaultDisplayRateProperty = LevelSequenceSettings->GetClass()->FindPropertyByName(GET_MEMBER_NAME_CHECKED(ULevelSequenceProjectSettings, DefaultDisplayRate));
@@ -122,6 +123,7 @@ void UProductionSettings::ApplySequencerSettingOverrides()
 	OverrideDefaultDisplayRate();
 	OverrideDefaultStartTime();
 	OverrideSubsequenceHierarchicalBias();
+	OverrideDefaultAssetNames();
 }
 
 void UProductionSettings::OverrideDefaultDisplayRate()
@@ -180,6 +182,30 @@ void UProductionSettings::OverrideSubsequenceHierarchicalBias()
 	}
 }
 
+void UProductionSettings::OverrideDefaultAssetNames()
+{
+	if (ActiveProduction)
+	{
+		IAssetTools& AssetTools = FModuleManager::GetModuleChecked<FAssetToolsModule>("AssetTools").Get();
+
+		// Register all of the active production's default asset names with the asset tools
+		for (const TPair<TObjectPtr<const UClass>, FString>& Naming : ActiveProduction->DefaultAssetNames)
+		{
+			if (Naming.Key)
+			{
+				// Before registering a new default name with AssetTools, cache the existing value (if it exists) so we can reset it later
+				TOptional<FString> ExistingDefaultName = AssetTools.GetDefaultAssetNameForClass(Naming.Key);
+				if (ExistingDefaultName.IsSet() && !ProjectDefaultAssetNames.Contains(Naming.Key))
+				{
+					ProjectDefaultAssetNames.Add(Naming.Key, ExistingDefaultName.GetValue());
+				}
+
+				AssetTools.RegisterDefaultAssetNameForClass(Naming.Key, Naming.Value);
+			}
+		}
+	}
+}
+
 void UProductionSettings::PostEditChangeProperty(struct FPropertyChangedEvent& PropertyChangedEvent)
 {
 	Super::PostEditChangeProperty(PropertyChangedEvent);
@@ -313,11 +339,30 @@ FCinematicProduction* UProductionSettings::GetActiveProduction() const
 
 void UProductionSettings::SetActiveProduction(FGuid ProductionID)
 {
+	if (ActiveProduction)
+	{
+		// Unregister any default asset names that the previously active production had set
+		// Then, restore the corresponding project default asset names
+		IAssetTools& AssetTools = FModuleManager::GetModuleChecked<FAssetToolsModule>("AssetTools").Get();
+
+		for (const TPair<TObjectPtr<const UClass>, FString>& Naming : ActiveProduction->DefaultAssetNames)
+		{
+			AssetTools.UnregisterDefaultAssetNameForClass(Naming.Key);
+
+			if (ProjectDefaultAssetNames.Contains(Naming.Key))
+			{
+				AssetTools.RegisterDefaultAssetNameForClass(Naming.Key, ProjectDefaultAssetNames[Naming.Key]);
+			}
+		}
+	}
+
 	ActiveProduction = Productions.FindByPredicate([ProductionID](const FCinematicProduction& Other) { return Other.ProductionID == ProductionID; });
 	SetActiveProductionName();
 
 	// Apply the overrides based on the new active production (even if it is None, because some properties will be reset to default)
-	ApplySequencerSettingOverrides();
+	ApplyProjectOverrides();
+
+	ActiveProductionChangedDelegate.Broadcast();
 }
 
 bool UProductionSettings::IsActiveProduction(FGuid ProductionID) const
@@ -448,6 +493,67 @@ void UProductionSettings::RemoveNamespaceFromDenyList(FGuid ProductionID, const
 	}
 }
 
+void UProductionSettings::AddAssetNaming(FGuid ProductionID, const UClass* AssetClass, const FString& DefaultName)
+{
+	if (!AssetClass)
+	{
+		return;
+	}
+
+	FCinematicProduction* Production = Productions.FindByPredicate([ProductionID](const FCinematicProduction& Other) { return Other.ProductionID == ProductionID; });
+	if (Production)
+	{
+		Production->DefaultAssetNames.Add(AssetClass, DefaultName);
+
+		if (IsActiveProduction(ProductionID))
+		{
+			IAssetTools& AssetTools = FModuleManager::GetModuleChecked<FAssetToolsModule>("AssetTools").Get();
+
+			// If there is an existing name registered with AssetTools for this class, cache it before overriding it
+			TOptional<FString> ExistingDefaultName = AssetTools.GetDefaultAssetNameForClass(AssetClass);
+			if (ExistingDefaultName.IsSet() && !ProjectDefaultAssetNames.Contains(AssetClass))
+			{
+				ProjectDefaultAssetNames.Add(AssetClass, ExistingDefaultName.GetValue());
+			}
+
+			// Register the new default asset name with AssetTools
+			AssetTools.RegisterDefaultAssetNameForClass(AssetClass, DefaultName);
+		}
+
+		TryUpdateDefaultConfigFile();
+	}
+}
+
+void UProductionSettings::RemoveAssetNaming(FGuid ProductionID, const UClass* AssetClass)
+{
+	if (!AssetClass)
+	{
+		return;
+	}
+
+	FCinematicProduction* Production = Productions.FindByPredicate([ProductionID](const FCinematicProduction& Other) { return Other.ProductionID == ProductionID; });
+	if (Production)
+	{
+		Production->DefaultAssetNames.Remove(AssetClass);
+
+		if (IsActiveProduction(ProductionID))
+		{
+			IAssetTools& AssetTools = FModuleManager::GetModuleChecked<FAssetToolsModule>("AssetTools").Get();
+
+			// Unregister the asset naming from AssetTools
+			AssetTools.UnregisterDefaultAssetNameForClass(AssetClass);
+
+			// If we had previously cached a default name (before this production overrode it) restore that name now
+			if (ProjectDefaultAssetNames.Contains(AssetClass))
+			{
+				AssetTools.RegisterDefaultAssetNameForClass(AssetClass, ProjectDefaultAssetNames[AssetClass]);
+			}
+		}
+
+		TryUpdateDefaultConfigFile();
+	}
+}
+
 void UProductionSettings::SetActiveDisplayRate(FFrameRate DisplayRate)
 {
 	if (ActiveProduction)
diff --git a/Engine/Plugins/VirtualProduction/Shots/Source/ShotsEditor/Private/ProductionSettings.h b/Engine/Plugins/VirtualProduction/Shots/Source/ShotsEditor/Private/ProductionSettings.h
--- a/Engine/Plugins/VirtualProduction/Shots/Source/ShotsEditor/Private/ProductionSettings.h
+++ b/Engine/Plugins/VirtualProduction/Shots/Source/ShotsEditor/Private/ProductionSettings.h
@@ -8,6 +8,7 @@
 #include "ProductionSettings.generated.h"
 
 DECLARE_MULTICAST_DELEGATE(FOnProductionListChanged);
+DECLARE_MULTICAST_DELEGATE(FOnActiveProductionChanged);
 
 /**
  * Options for determining the hierarchical bias of subsequences
@@ -53,6 +54,10 @@ public:
 	/** List of Naming Token namespaces that should not be evaluated */
 	UPROPERTY(config, EditAnywhere, BlueprintReadOnly, Category = "Default")
 	TSet<FString> NamingTokenNamespaceDenyList;
+
+	/** List of default names for specific asset types */
+	UPROPERTY(config, VisibleAnywhere, BlueprintReadOnly, Category = "Default")
+	TMap<TObjectPtr<const UClass>, FString> DefaultAssetNames;
 };
 
 /**
@@ -122,6 +127,12 @@ public:
 	/** Removes a Naming Token namespace from the DenyList of the production matching the input ID */
 	void RemoveNamespaceFromDenyList(FGuid ProductionID, const FString& Namespace);
 
+	/** Adds a new entry into the DefaultAssetNames map of the production matching the input ID */
+	void AddAssetNaming(FGuid ProductionID, const UClass* AssetClass, const FString& DefaultName);
+
+	/** Removes an entry from the DefaultAssetNames map of the production matching the input ID */
+	void RemoveAssetNaming(FGuid ProductionID, const UClass* AssetClass);
+
 	/** Sets the DefaultDisplayRate of the active production */
 	void SetActiveDisplayRate(FFrameRate DisplayRate);
 
@@ -137,12 +148,15 @@ public:
 	/** Returns the delegate that broadcasts when a production is added/removed */
 	FOnProductionListChanged& OnProductionListChanged() { return ProductionListChangedDelegate; }
 
+	/** Returns the delegate that broadcasts when the active production changes */
+	FOnActiveProductionChanged& OnActiveProductionChanged() { return ActiveProductionChangedDelegate; }
+
 private:
 	/** Callback when a property of ULevelSequenceProjectSettings has changed */
 	void OnLevelSequenceSettingChanged(UObject* Object, struct FPropertyChangedEvent& PropertyChangedEvent);
 
-	/** Applies overrides to various sequencer settings based on the active production settings */
-	void ApplySequencerSettingOverrides();
+	/** Applies overrides to various project settings based on the active production settings */
+	void ApplyProjectOverrides();
 
 	/** Overrides the DefaultDisplayRate in the level sequence project settings based on the active production setting */
 	void OverrideDefaultDisplayRate();
@@ -153,6 +167,9 @@ private:
 	/** Writes out the Hierarchical bias value (based on the Subsequence Priority setting) to EditorPerProjectUserSettings.ini */
 	void OverrideSubsequenceHierarchicalBias();
 
+	/** Overrides the DefaultAssetNames property of the asset tools project settings based on the active production setting */
+	void OverrideDefaultAssetNames();
+
 	/** Set the active production name and write out the new value to EditorPerProjectUserSettings.ini */
 	void SetActiveProductionName();
 
@@ -172,10 +189,16 @@ private:
 	/** Default value of DefaultDisplayRate in the level sequence project setting, used to reset it when there is no active production */
 	FString DefaultDisplayRateProjectSetting;
 
+	/** Default asset names previously registered with AssetTools, used to reset them when the active production changes */
+	TMap<const UClass*, FString> ProjectDefaultAssetNames;
+
 	/** Original tooltip text for sequencer settings */
 	FString OriginalDefaultDisplayRateTooltip;
 	FString OriginalDefaultStartTimeTooltip;
 
 	/** Delegate that broadcasts when a production is added/removed */
 	FOnProductionListChanged ProductionListChangedDelegate;
+
+	/** Delegate that broadcasts when the active production changes */
+	FOnActiveProductionChanged ActiveProductionChangedDelegate;
 };
diff --git a/Engine/Plugins/VirtualProduction/Shots/Source/ShotsEditor/Private/ShotsEditorStyle.cpp b/Engine/Plugins/VirtualProduction/Shots/Source/ShotsEditor/Private/ShotsEditorStyle.cpp
--- a/Engine/Plugins/VirtualProduction/Shots/Source/ShotsEditor/Private/ShotsEditorStyle.cpp
+++ b/Engine/Plugins/VirtualProduction/Shots/Source/ShotsEditor/Private/ShotsEditorStyle.cpp
@@ -51,6 +51,14 @@ FShotsEditorStyle::FShotsEditorStyle()
 
 	Set("ProductionWizard.RecessedButton", RecessedButton);
 
+	const FButtonStyle PanelButton = FButtonStyle()
+		.SetNormal(FSlateRoundedBoxBrush(FStyleColors::Panel, 0.0f, FStyleColors::Panel, 0.0f))
+		.SetHovered(FSlateRoundedBoxBrush(FStyleColors::Header, 0.0f, FStyleColors::Header, 0.0f))
+		.SetPressed(FSlateRoundedBoxBrush(FStyleColors::Dropdown, 0.0f, FStyleColors::Dropdown, 0.0f))
+		.SetDisabled(FSlateRoundedBoxBrush(FStyleColors::Panel, 0.0f, FStyleColors::Panel, 0.0f));
+
+	Set("ProductionWizard.PanelButton", PanelButton);
+
 	FSlateStyleRegistry::RegisterSlateStyle(*this);
 }
 
diff --git a/Engine/Plugins/VirtualProduction/Shots/Source/ShotsEditor/Private/UI/SActiveProductionCombo.cpp b/Engine/Plugins/VirtualProduction/Shots/Source/ShotsEditor/Private/UI/SActiveProductionCombo.cpp
--- a/Engine/Plugins/VirtualProduction/Shots/Source/ShotsEditor/Private/UI/SActiveProductionCombo.cpp
+++ b/Engine/Plugins/VirtualProduction/Shots/Source/ShotsEditor/Private/UI/SActiveProductionCombo.cpp
@@ -15,6 +15,10 @@ void SActiveProductionCombo::Construct(const FArguments& InArgs)
 	static auto GetActiveProductionName = []() -> FText
 		{
 			const UProductionSettings* const ProductionSettings = GetDefault<UProductionSettings>();
+			if (ProductionSettings->ActiveProductionName.IsEmpty())
+			{
+				return LOCTEXT("NoneProductionName", "None");
+			}
 			return FText::FromString(ProductionSettings->ActiveProductionName);
 		};
 
diff --git a/Engine/Plugins/VirtualProduction/Shots/Source/ShotsEditor/Private/UI/SAssetNamingPanel.cpp b/Engine/Plugins/VirtualProduction/Shots/Source/ShotsEditor/Private/UI/SAssetNamingPanel.cpp
new file mode 100644
--- /dev/null
+++ b/Engine/Plugins/VirtualProduction/Shots/Source/ShotsEditor/Private/UI/SAssetNamingPanel.cpp
@@ -0,0 +1,437 @@
+// Copyright Epic Games, Inc. All Rights Reserved.
+
+#include "SAssetNamingPanel.h"
+
+#include "AssetDefinitionRegistry.h"
+#include "AssetToolsModule.h"
+#include "ClassViewerFilter.h"
+#include "Factories/Factory.h"
+#include "Framework/Application/MenuStack.h"
+#include "Framework/Application/SlateApplication.h"
+#include "Framework/MultiBox/MultiBoxBuilder.h"
+#include "ProductionSettings.h"
+#include "PropertyCustomizationHelpers.h"
+#include "ShotsEditorStyle.h"
+#include "Styling/StyleColors.h"
+#include "UI/SActiveProductionCombo.h"
+#include "Widgets/Layout/SSeparator.h"
+#include "Widgets/SBoxPanel.h"
+#include "Widgets/Text/SInlineEditableTextBlock.h"
+
+#define LOCTEXT_NAMESPACE "SAssetNamingPanel"
+
+namespace UE::Shots::AssetNaming::Private
+{
+	// Filter for all classes that have a registered asset definition
+	// This roughly approximates anything that the user can create in the content browser, and that they might want to assign a default name
+	class FAssetNamingClassFilter : public IClassViewerFilter
+	{
+	public:
+		FAssetNamingClassFilter()
+		{
+			TArray<TSoftClassPtr<UObject>> RegisteredAssetClasses = UAssetDefinitionRegistry::Get()->GetAllRegisteredAssetClasses();
+			for (const TSoftClassPtr<UObject>& AsssetClass : RegisteredAssetClasses)
+			{
+				AllowedAssetClasses.Add(AsssetClass.Get());
+			}
+		}
+
+		virtual bool IsClassAllowed(const FClassViewerInitializationOptions& InInitOptions, const UClass* InClass, TSharedRef<FClassViewerFilterFuncs> InFilterFuncs) override
+		{
+			return AllowedAssetClasses.Contains(InClass);
+		}
+
+		virtual bool IsUnloadedClassAllowed(const FClassViewerInitializationOptions& InInitOptions, const TSharedRef<const class IUnloadedBlueprintData> InUnloadedClassData, TSharedRef<class FClassViewerFilterFuncs> InFilterFuncs) override
+		{
+			return false;
+		}
+
+	private:
+		TArray<UClass*> AllowedAssetClasses;
+	};
+
+	static FName ClassColumnName = TEXT("Class");
+	static FName NamingColumnName = TEXT("Naming");
+	static FName DeleteColumnName = TEXT("Delete");
+}
+
+void SAssetNamingPanel::Construct(const FArguments& InArgs)
+{
+	using namespace UE::Shots::AssetNaming::Private;
+
+	// Subscribe to be notified when the Production Settings active productions has changed
+	UProductionSettings* ProductionSettings = GetMutableDefault<UProductionSettings>();
+	ActiveProductionChangedHandle = ProductionSettings->OnActiveProductionChanged().AddSP(this, &SAssetNamingPanel::UpdateAssetNamingList);
+
+	// Initialize the asset naming list items for the current active production
+	UpdateAssetNamingList();
+
+	// Lambdas associate with the "Create New Naming" button
+	auto IsActiveProductionValid = []() -> bool
+		{
+			const UProductionSettings* const ProductionSettings = GetDefault<UProductionSettings>();
+			return (ProductionSettings->GetActiveProduction() != nullptr);
+		};
+
+	auto OnCreateNewNaming = [this]() -> FReply
+		{
+			AssetNamingListItems.Add(MakeShared<FAssetNamingRowData>());
+
+			if (AssetNamingListView)
+			{
+				AssetNamingListView->RequestListRefresh();
+			}
+
+			return FReply::Handled();
+		};
+
+	ChildSlot
+		[
+			SNew(SVerticalBox)
+
+			// Active Production Selector
+			+ SVerticalBox::Slot()
+				.AutoHeight()
+				[
+					SNew(SActiveProductionCombo)
+				]
+
+			// Separator
+			+ SVerticalBox::Slot()
+				.AutoHeight()
+				[
+					SNew(SSeparator)
+						.Orientation(Orient_Horizontal)
+						.Thickness(2.0f)
+				]
+
+			// Asset Naming Panel
+			+ SVerticalBox::Slot()
+				.FillHeight(1.0f)
+				[
+					SNew(SBorder)
+						.BorderImage(FAppStyle::Get().GetBrush("Brushes.Panel"))
+						.Padding(16.0f)
+						[
+							SNew(SVerticalBox)
+
+							// Title
+							+ SVerticalBox::Slot()
+								.AutoHeight()
+								.Padding(0.0f, 0.0f, 0.0f, 4.0f)
+								[
+									SNew(STextBlock)
+										.Text(LOCTEXT("AssetNamingTitle", "Production Settings"))
+										.Font(FShotsEditorStyle::Get().GetFontStyle("ProductionWizard.TitleFont"))
+								]
+
+							// Heading
+							+ SVerticalBox::Slot()
+								.AutoHeight()
+								.Padding(0.0f, 0.0f, 0.0f, 4.0f)
+								[
+									SNew(STextBlock)
+										.Text(LOCTEXT("AssetNamingHeading", "Asset Naming"))
+										.Font(FShotsEditorStyle::Get().GetFontStyle("ProductionWizard.HeadingFont"))
+								]
+
+							// Info Text
+							+ SVerticalBox::Slot()
+								.AutoHeight()
+								.Padding(0.0f, 0.0f, 0.0f, 16.0f)
+								[
+									SNew(STextBlock)
+										.Text(LOCTEXT("AssetNamingInfoText", "For each asset class, you can set a prefix for the editor to assign automatically to the names of newly created assets. This helps keep asset names consistent across your project."))
+										.AutoWrapText(true)
+								]
+
+							// Create New Naming Button
+							+ SVerticalBox::Slot()
+								.AutoHeight()
+								.HAlign(HAlign_Left)
+								.Padding(0.0f, 0.0f, 0.0f, 8.0f)
+								[
+									SNew(SButton)
+										.ContentPadding(FMargin(2.0f))
+										.OnClicked_Lambda(OnCreateNewNaming)
+										.IsEnabled_Lambda(IsActiveProductionValid)
+										[
+											SNew(SHorizontalBox)
+
+											+ SHorizontalBox::Slot()
+												.AutoWidth()
+												.Padding(0.0f, 0.0f, 4.0f, 0.0f)
+												[
+													SNew(SImage)
+														.Image(FAppStyle::Get().GetBrush("Icons.Plus"))
+														.ColorAndOpacity(FStyleColors::AccentGreen)
+												]
+
+											+ SHorizontalBox::Slot()
+												.AutoWidth()
+												[
+													SNew(STextBlock).Text(LOCTEXT("CreateNewNamingButton", "Create New Naming"))
+												]
+										]
+								]
+
+							// Asset Naming List View
+							+ SVerticalBox::Slot()
+								.FillHeight(1.0f)
+								[
+									SAssignNew(AssetNamingListView, SListView<TSharedPtr<FAssetNamingRowData>>)
+										.ListItemsSource(&AssetNamingListItems)
+										.OnGenerateRow(this, &SAssetNamingPanel::OnGenerateAssetNamingRow)
+										.SelectionMode(ESelectionMode::None)
+										.HeaderRow
+										(
+											SNew(SHeaderRow)
+
+											+ SHeaderRow::Column(ClassColumnName)
+												.DefaultLabel(LOCTEXT("ClassLabel", "Class"))
+												.FillWidth(0.25f)
+
+											+ SHeaderRow::Column(NamingColumnName)
+												.DefaultLabel(LOCTEXT("NamingLabel", "Default Naming"))
+												.FillWidth(0.7f)
+
+											+ SHeaderRow::Column(DeleteColumnName)
+												.DefaultLabel(FText::GetEmpty())
+												.FillWidth(0.05f)
+										)
+								]
+						]
+				]
+		];
+}
+
+SAssetNamingPanel::~SAssetNamingPanel()
+{
+	if (UObjectInitialized())
+	{
+		if (UProductionSettings* ProductionSettings = GetMutableDefault<UProductionSettings>())
+		{
+			ProductionSettings->OnActiveProductionChanged().Remove(ActiveProductionChangedHandle);
+		}
+	}
+}
+
+void SAssetNamingPanel::UpdateAssetNamingList()
+{
+	AssetNamingListItems.Empty();
+
+	const UProductionSettings* const ProductionSettings = GetDefault<UProductionSettings>();
+	if (const FCinematicProduction* ActiveProduction = ProductionSettings->GetActiveProduction())
+	{
+		for (const TPair<TObjectPtr<const UClass>, FString>& AssetNaming : ActiveProduction->DefaultAssetNames)
+		{
+			AssetNamingListItems.Add(MakeShared<FAssetNamingRowData>(AssetNaming.Key, AssetNaming.Value));
+		}
+	}
+
+	if (AssetNamingListView)
+	{
+		AssetNamingListView->RequestListRefresh();
+	}
+}
+
+TSharedRef<ITableRow> SAssetNamingPanel::OnGenerateAssetNamingRow(TSharedPtr<FAssetNamingRowData> InAssetNaming, const TSharedRef<STableViewBase>& OwnerTable)
+{
+	return SNew(SAssetNamingRow, OwnerTable, InAssetNaming)
+		.OnDeleteRow(this, &SAssetNamingPanel::UpdateAssetNamingList);
+}
+
+void SAssetNamingRow::Construct(const FArguments& Args, const TSharedRef<STableViewBase>& OwnerTableView, const TSharedPtr<FAssetNamingRowData>& InAssetNaming)
+{
+	AssetNaming = InAssetNaming;
+
+	OnDeleteRow = Args._OnDeleteRow;
+
+	SMultiColumnTableRow<TSharedPtr<FAssetNamingRowData>>::Construct(FSuperRowType::FArguments(), OwnerTableView);
+}
+
+TSharedRef<SWidget> SAssetNamingRow::GenerateWidgetForColumn(const FName& ColumnName)
+{
+	using namespace UE::Shots::AssetNaming::Private;
+
+	if (ColumnName == ClassColumnName)
+	{
+		return SNew(SHorizontalBox)
+
+		+ SHorizontalBox::Slot()
+			.Padding(0.0f, 0.0f, 0.0f, 1.0f)
+			[
+				SNew(SBorder)
+					.BorderImage(FAppStyle::Get().GetBrush("Brushes.Panel"))
+					.Padding(FMargin(4.0, 1.0f))
+					[
+						SNew(SClassPropertyEntryBox)
+							.MetaClass(UObject::StaticClass())
+							.SelectedClass_Lambda([this]() { return AssetNaming->Class; })
+							.ClassViewerFilters({ MakeShared<FAssetNamingClassFilter>() })
+							.OnSetClass(this, &SAssetNamingRow::SetAssetClass)
+					]
+			];
+	}
+	else if (ColumnName == NamingColumnName)
+	{
+		return SNew(SHorizontalBox)
+
+			+ SHorizontalBox::Slot()
+			.Padding(1.0f, 0.0f, 0.0f, 1.0f)
+			[
+				SNew(SBorder)
+					.BorderImage(FAppStyle::Get().GetBrush("Brushes.Panel"))
+					.Padding(FMargin(4.0, 1.0f))
+					.OnMouseButtonDown(this, &SAssetNamingRow::SummonEditMenu)
+					[
+						SAssignNew(EditableTextBlock, SInlineEditableTextBlock)
+							.Text_Lambda([this]() { return FText::FromString(AssetNaming->DefaultName); })
+							.IsEnabled_Lambda([this]() { return (AssetNaming->Class != nullptr); })
+							.OnVerifyTextChanged(this, &SAssetNamingRow::ValidateDefaultAssetName)
+							.OnTextCommitted(this, &SAssetNamingRow::SetDefaultAssetName)
+					]
+			];
+	}
+	else if (ColumnName == DeleteColumnName)
+	{
+		return SNew(SHorizontalBox)
+
+		+ SHorizontalBox::Slot()
+			.Padding(1.0f, 0.0f, 0.0f, 1.0f)
+			[
+				SNew(SButton)
+					.ButtonStyle(FShotsEditorStyle::Get(), "ProductionWizard.PanelButton")
+					.ContentPadding(FMargin(0.0f, 0.0f))
+					.HAlign(HAlign_Center)
+					.VAlign(VAlign_Center)
+					.OnClicked(this, &SAssetNamingRow::DeleteRow)
+					.Content()
+					[
+						SNew(SImage).Image(FAppStyle::Get().GetBrush("Icons.Delete"))
+					]
+			];
+	}
+
+	return SNullWidget::NullWidget;
+}
+
+void SAssetNamingRow::SetAssetClass(const UClass* SelectedClass)
+{
+	// Early-out if the same class was just selected
+	if (AssetNaming->Class == SelectedClass)
+	{
+		return;
+	}
+
+	// Before updating the asset class for this row, remove the old class from the active production settings list
+	UProductionSettings* ProductionSettings = GetMutableDefault<UProductionSettings>();
+	if (FCinematicProduction* ActiveProduction = ProductionSettings->GetActiveProduction())
+	{
+		ProductionSettings->RemoveAssetNaming(ActiveProduction->ProductionID, AssetNaming->Class);
+	}
+
+	AssetNaming->Class = SelectedClass;
+
+	if (SelectedClass)
+	{	
+		// Set the default name for this row to the factory default name for the selected asset type
+		const TArray<UFactory*>& Factories = PropertyCustomizationHelpers::GetNewAssetFactoriesForClasses({ SelectedClass });
+		if (Factories.Num() > 0)
+		{
+			AssetNaming->DefaultName = Factories[0]->GetDefaultNewAssetName();
+		}
+
+		// Let the user immediately start editing the asset name
+		FSlateApplication::Get().SetKeyboardFocus(EditableTextBlock.ToSharedRef());
+		EditableTextBlock->EnterEditingMode();
+	}
+	else
+	{
+		AssetNaming->DefaultName.Empty();
+	}
+}
+
+void SAssetNamingRow::SetDefaultAssetName(const FText& InText, ETextCommit::Type InCommitType)
+{
+	AssetNaming->DefaultName = InText.ToString();
+
+	UProductionSettings* ProductionSettings = GetMutableDefault<UProductionSettings>();
+	if (FCinematicProduction* ActiveProduction = ProductionSettings->GetActiveProduction())
+	{
+		ProductionSettings->AddAssetNaming(ActiveProduction->ProductionID, AssetNaming->Class, AssetNaming->DefaultName);
+	}
+}
+
+bool SAssetNamingRow::ValidateDefaultAssetName(const FText& InText, FText& OutErrorMessage)
+{
+	// An empty name is invalid
+	if (InText.IsEmpty())
+	{
+		OutErrorMessage = LOCTEXT("EmptyNameErrorMessage", "Please provide a default asset name");
+		return false;
+	}
+
+	// Ensure that the name does not contain any characters that would be invalid for an asset name
+	// This matches the validation that would happen if the user was renaming an asset in the content browser
+	return FName::IsValidXName(InText.ToString(), INVALID_OBJECTNAME_CHARACTERS INVALID_LONGPACKAGE_CHARACTERS, &OutErrorMessage);
+}
+
+FReply SAssetNamingRow::DeleteRow()
+{
+	UProductionSettings* ProductionSettings = GetMutableDefault<UProductionSettings>();
+	if (FCinematicProduction* ActiveProduction = ProductionSettings->GetActiveProduction())
+	{
+		ProductionSettings->RemoveAssetNaming(ActiveProduction->ProductionID, AssetNaming->Class);
+	}
+
+	OnDeleteRow.ExecuteIfBound();
+
+	return FReply::Handled();
+}
+
+FReply SAssetNamingRow::SummonEditMenu(const FGeometry& MyGeometry, const FPointerEvent& MouseEvent)
+{
+	// Create the context menu to be launched on right mouse click
+	if (MouseEvent.GetEffectingButton() == EKeys::RightMouseButton)
+	{
+		FMenuBuilder MenuBuilder(true, nullptr);
+
+		MenuBuilder.AddMenuEntry(
+			LOCTEXT("RenameAction", "Rename"),
+			LOCTEXT("RenameActionToolTip", "Change the default name for this asset type"),
+			FSlateIcon(FShotsEditorStyle::StyleName, "Icons.AssetNaming"),
+			FUIAction(FExecuteAction::CreateLambda([this]()
+				{
+					FSlateApplication::Get().SetKeyboardFocus(EditableTextBlock.ToSharedRef());
+					EditableTextBlock->EnterEditingMode();
+				}
+			)),
+			NAME_None,
+			EUserInterfaceActionType::Button
+		);
+
+		MenuBuilder.AddMenuEntry(
+			LOCTEXT("DeleteAction", "Delete"),
+			LOCTEXT("DeleteActionToolTip", "Delete this default naming entry"),
+			FSlateIcon(FAppStyle::GetAppStyleSetName(), "Icons.Delete"),
+			FUIAction(FExecuteAction::CreateLambda([this]() { DeleteRow(); } )),
+			NAME_None,
+			EUserInterfaceActionType::Button
+		);
+
+		FWidgetPath WidgetPath = MouseEvent.GetEventPath() != nullptr ? *MouseEvent.GetEventPath() : FWidgetPath();
+
+		FSlateApplication::Get().PushMenu(
+			AsShared(),
+			WidgetPath,
+			MenuBuilder.MakeWidget(),
+			MouseEvent.GetScreenSpacePosition(),
+			FPopupTransitionEffect(FPopupTransitionEffect::ContextMenu)
+		);
+	}
+
+	return FReply::Handled();
+}
+
+#undef LOCTEXT_NAMESPACE
diff --git a/Engine/Plugins/VirtualProduction/Shots/Source/ShotsEditor/Private/UI/SAssetNamingPanel.h b/Engine/Plugins/VirtualProduction/Shots/Source/ShotsEditor/Private/UI/SAssetNamingPanel.h
new file mode 100644
--- /dev/null
+++ b/Engine/Plugins/VirtualProduction/Shots/Source/ShotsEditor/Private/UI/SAssetNamingPanel.h
@@ -0,0 +1,99 @@
+// Copyright Epic Games, Inc. All Rights Reserved.
+
+#pragma once
+
+#include "Widgets/SCompoundWidget.h"
+
+#include "Widgets/DeclarativeSyntaxSupport.h"
+#include "Widgets/Text/SInlineEditableTextBlock.h"
+#include "Widgets/Views/SListView.h"
+
+/** Data to display in a row in the asset naming list view */
+struct FAssetNamingRowData
+{
+	/** The asset class being given a new default asset name (chosen using a class picker widget) */
+	const UClass* Class;
+
+	/** The new default asset name to register for the asset type */
+	FString DefaultName;
+};
+
+DECLARE_DELEGATE(FOnDeleteRow);
+
+/** Row widget for the asset naming list view */
+class SAssetNamingRow : public SMultiColumnTableRow<TSharedPtr<FAssetNamingRowData>>
+{
+public:
+
+	SLATE_BEGIN_ARGS(SAssetNamingRow) 
+		: _OnDeleteRow()
+		{}
+
+		/** Called when a row is deleted from the list */
+		SLATE_EVENT(FOnDeleteRow, OnDeleteRow)
+
+	SLATE_END_ARGS()
+
+	void Construct(const FArguments& Args, const TSharedRef<STableViewBase>& OwnerTableView, const TSharedPtr<FAssetNamingRowData>& InRowData);
+
+	/** Creates the widget for this row for the specified column */
+	virtual TSharedRef<SWidget> GenerateWidgetForColumn(const FName& ColumnName) override;
+
+private:
+	/** Set the asset class for this row widget */
+	void SetAssetClass(const UClass* SelectedClass);
+
+	/** Set the default name for this row widget */
+	void SetDefaultAssetName(const FText& InText, ETextCommit::Type InCommitType);
+
+	/** Validate the text entered by the user to ensure it will be a valid asset name */
+	bool ValidateDefaultAssetName(const FText& InText, FText& OutErrorMessage);
+
+	/** Removes this row's default naming from the active production settings, and removes it from the list view */
+	FReply DeleteRow();
+
+	/** Summons a right-click context menu for the current row */
+	FReply SummonEditMenu(const FGeometry& MyGeometry, const FPointerEvent& MouseEvent);
+
+private:
+	/** The underlying asset naming data, used to properly display the class picker and default naming for that class */
+	TSharedPtr<FAssetNamingRowData> AssetNaming;
+
+	/** The text box where the user can type the default asset name */
+	TSharedPtr<SInlineEditableTextBlock> EditableTextBlock;
+
+	/** Delegate to execute when a row is deleted from the list */
+	FOnDeleteRow OnDeleteRow;
+};
+
+/**
+ * UI for the Asset Naming panel in the Production Wizard
+ */
+class SAssetNamingPanel : public SCompoundWidget
+{
+public:
+	SAssetNamingPanel() = default;
+	~SAssetNamingPanel();
+
+	SLATE_BEGIN_ARGS(SAssetNamingPanel) {}
+	SLATE_END_ARGS()
+
+	void Construct(const FArguments& InArgs);
+
+private:
+	/** Generates a row displaying asset naming data */
+	TSharedRef<ITableRow> OnGenerateAssetNamingRow(TSharedPtr<FAssetNamingRowData> InAssetNaming, const TSharedRef<STableViewBase>& OwnerTable);
+
+	/** Update the list view with the asset naming properties of the current active production */
+	void UpdateAssetNamingList();
+
+private:
+	/** Source of default asset names for the active production */
+	TArray<TSharedPtr<FAssetNamingRowData>> AssetNamingListItems;
+
+	/** List view displaying the default asset names for the active production */
+	TSharedPtr<SListView<TSharedPtr<FAssetNamingRowData>>> AssetNamingListView;
+
+	/** Delegate bound to the Production Setting's OnActiveProductionChanged event */
+	FDelegateHandle ActiveProductionChangedHandle;
+};
diff --git a/Engine/Plugins/VirtualProduction/Shots/Source/ShotsEditor/Private/UI/SProductionWizard.cpp b/Engine/Plugins/VirtualProduction/Shots/Source/ShotsEditor/Private/UI/SProductionWizard.cpp
--- a/Engine/Plugins/VirtualProduction/Shots/Source/ShotsEditor/Private/UI/SProductionWizard.cpp
+++ b/Engine/Plugins/VirtualProduction/Shots/Source/ShotsEditor/Private/UI/SProductionWizard.cpp
@@ -5,6 +5,7 @@
 #include "ProductionSettings.h"
 #include "ShotsEditorStyle.h"
 #include "UI/SNamingTokensPanel.h"
+#include "UI/SAssetNamingPanel.h"
 #include "UI/SProductionListPanel.h"
 #include "UI/SSequencerSettingsPanel.h"
 #include "Widgets/Layout/SBorder.h"
@@ -124,7 +125,7 @@ TSharedRef<SWidget> SProductionWizard::MakeMenuPanel()
 	TSharedPtr<FProductionWizardMenuEntry> AssetNamingMenuEntry = MakeShared<FProductionWizardMenuEntry>();
 	AssetNamingMenuEntry->Label = LOCTEXT("AssetNamingMenuLabel", "Asset Naming");
 	AssetNamingMenuEntry->Icon = FSlateIcon(FShotsEditorStyle::StyleName, "Icons.AssetNaming");
-	AssetNamingMenuEntry->Panel = MakeAssetNamingPanel();
+	AssetNamingMenuEntry->Panel = SNew(SAssetNamingPanel);
 
 	// Build the Folder Hierarchy menu entry, including its content panel
 	TSharedPtr<FProductionWizardMenuEntry> FolderHierarchyMenuEntry = MakeShared<FProductionWizardMenuEntry>();
@@ -258,26 +259,6 @@ TSharedRef<SWidget> SProductionWizard::MakeRevisionControlPanel()
 		];
 }
 
-TSharedRef<SWidget> SProductionWizard::MakeAssetNamingPanel()
-{
-	// TODO: Fill in this stubbed out panel
-	return SNew(SVerticalBox)
-
-	+ SVerticalBox::Slot()
-		.FillHeight(1.0f)
-		.HAlign(HAlign_Center)
-		[
-			SNew(SHorizontalBox)
-
-			+ SHorizontalBox::Slot()
-				.FillWidth(1.0f)
-				.VAlign(VAlign_Center)
-				[
-					SNew(STextBlock).Text(LOCTEXT("AssetNamingPanel", "ASSET NAMING PANEL"))
-				]
-		];
-}
-
 TSharedRef<SWidget> SProductionWizard::MakeFolderHierarchyPanel()
 {
 	// TODO: Fill in this stubbed out panel
diff --git a/Engine/Plugins/VirtualProduction/Shots/Source/ShotsEditor/Private/UI/SProductionWizard.h b/Engine/Plugins/VirtualProduction/Shots/Source/ShotsEditor/Private/UI/SProductionWizard.h
--- a/Engine/Plugins/VirtualProduction/Shots/Source/ShotsEditor/Private/UI/SProductionWizard.h
+++ b/Engine/Plugins/VirtualProduction/Shots/Source/ShotsEditor/Private/UI/SProductionWizard.h
@@ -44,7 +44,6 @@ private:
 
 	/** Functions for creating the panels for each of the production wizard's menu entries */
 	TSharedRef<SWidget> MakeRevisionControlPanel();
-	TSharedRef<SWidget> MakeAssetNamingPanel();
 	TSharedRef<SWidget> MakeFolderHierarchyPanel();
 
 	/** Generates a row in the menu list views */
diff --git a/Engine/Plugins/VirtualProduction/Shots/Source/ShotsEditor/ShotsEditor.build.cs b/Engine/Plugins/VirtualProduction/Shots/Source/ShotsEditor/ShotsEditor.build.cs
--- a/Engine/Plugins/VirtualProduction/Shots/Source/ShotsEditor/ShotsEditor.build.cs
+++ b/Engine/Plugins/VirtualProduction/Shots/Source/ShotsEditor/ShotsEditor.build.cs
@@ -15,6 +15,9 @@ namespace UnrealBuildTool.Rules
 			PrivateDependencyModuleNames.AddRange(
 				new string[]
 				{
+					"AssetDefinition",
+					"AssetTools",
+					"ClassViewer",
 					"Core",
 					"CoreUObject",
 					"DeveloperSettings",
@@ -26,6 +29,7 @@ namespace UnrealBuildTool.Rules
 					"MovieSceneTools",
 					"NamingTokens",
 					"Projects",
+					"PropertyEditor",
 					"Slate",
 					"SlateCore",
 					"SourceControl",
