ShotManagement: Make back and next buttons in production wizard functional

#rnx

diff --git a/Engine/Plugins/VirtualProduction/Shots/Source/ShotsEditor/Private/UI/SProductionWizard.cpp b/Engine/Plugins/VirtualProduction/Shots/Source/ShotsEditor/Private/UI/SProductionWizard.cpp
--- a/Engine/Plugins/VirtualProduction/Shots/Source/ShotsEditor/Private/UI/SProductionWizard.cpp
+++ b/Engine/Plugins/VirtualProduction/Shots/Source/ShotsEditor/Private/UI/SProductionWizard.cpp
@@ -62,6 +62,33 @@ void SProductionWizard::Construct(const FArguments& InArgs)
 
 TSharedRef<SWidget> SProductionWizard::MakeButtonsPanel()
 {
+	auto IsBackButtonEnabled = [this]()
+		{
+			TArray<TSharedPtr<FProductionWizardMenuEntry>> SelectedUserMenuEntries = UserMenuListView->GetSelectedItems();
+			if (SelectedUserMenuEntries.Num() > 0)
+			{
+				if (UserMenuEntries.IndexOfByKey(SelectedUserMenuEntries[0]) == 0)
+				{
+					return false;
+				}
+			}
+			return true;
+		};
+
+
+	auto IsNextButtonEnabled = [this]()
+		{
+			TArray<TSharedPtr<FProductionWizardMenuEntry>> SelectedProductionMenuEntries = ProductionMenuListView->GetSelectedItems();
+			if (SelectedProductionMenuEntries.Num() > 0)
+			{
+				if (ProductionMenuEntries.IndexOfByKey(SelectedProductionMenuEntries[0]) == ProductionMenuEntries.Num() - 1)
+				{
+					return false;
+				}
+			}
+			return true;
+		};
+
 	return SNew(SBorder)
 		.BorderImage(FAppStyle::Get().GetBrush("Brushes.Panel"))
 		.Padding(16.0f)
@@ -78,6 +105,8 @@ TSharedRef<SWidget> SProductionWizard::MakeButtonsPanel()
 					SNew(SButton)
 						.Text(LOCTEXT("BackButton", "Back"))
 						.HAlign(HAlign_Center)
+						.OnClicked(this, &SProductionWizard::OnBackClicked)
+						.IsEnabled_Lambda(IsBackButtonEnabled)
 				]
 
 			+ SHorizontalBox::Slot()
@@ -88,10 +117,72 @@ TSharedRef<SWidget> SProductionWizard::MakeButtonsPanel()
 					SNew(SButton)
 						.Text(LOCTEXT("NextButton", "Next"))
 						.HAlign(HAlign_Center)
+						.OnClicked(this, &SProductionWizard::OnNextClicked)
+						.IsEnabled_Lambda(IsNextButtonEnabled)
 				]
 		];
 }
 
+FReply SProductionWizard::OnBackClicked()
+{
+	// Both list views are single-select only, so they should only ever have 0 or 1 item selected
+	TArray<TSharedPtr<FProductionWizardMenuEntry>> SelectedUserMenuEntries = UserMenuListView->GetSelectedItems();
+	TArray<TSharedPtr<FProductionWizardMenuEntry>> SelectedProductionMenuEntries = ProductionMenuListView->GetSelectedItems();
+
+	if (SelectedUserMenuEntries.Num() > 0)
+	{
+		const int32 PreviousEntryIndex = UserMenuEntries.IndexOfByKey(SelectedUserMenuEntries[0]) - 1;
+		if (PreviousEntryIndex >= 0)
+		{
+			UserMenuListView->SetSelection(UserMenuEntries[PreviousEntryIndex]);
+		}
+	}
+	else if (SelectedProductionMenuEntries.Num() > 0)
+	{
+		const int32 PreviousEntryIndex = ProductionMenuEntries.IndexOfByKey(SelectedProductionMenuEntries[0]) - 1;
+		if (PreviousEntryIndex >= 0)
+		{
+			ProductionMenuListView->SetSelection(ProductionMenuEntries[PreviousEntryIndex]);
+		}
+		else
+		{
+			UserMenuListView->SetSelection(UserMenuEntries.Last());
+		}
+	}
+
+	return FReply::Handled();
+}
+
+FReply SProductionWizard::OnNextClicked()
+{
+	// Both list views are single-select only, so they should only ever have 0 or 1 item selected
+	TArray<TSharedPtr<FProductionWizardMenuEntry>> SelectedUserMenuEntries = UserMenuListView->GetSelectedItems();
+	TArray<TSharedPtr<FProductionWizardMenuEntry>> SelectedProductionMenuEntries = ProductionMenuListView->GetSelectedItems();
+
+	if (SelectedUserMenuEntries.Num() > 0)
+	{
+		const int32 NextIndex = UserMenuEntries.IndexOfByKey(SelectedUserMenuEntries[0]) + 1;
+		if (NextIndex < UserMenuEntries.Num())
+		{
+			UserMenuListView->SetSelection(UserMenuEntries[NextIndex]);
+		}
+		else
+		{
+			ProductionMenuListView->SetSelection(ProductionMenuEntries[0]);
+		}
+	}
+	else if (SelectedProductionMenuEntries.Num() > 0)
+	{
+		const int32 NextIndex = ProductionMenuEntries.IndexOfByKey(SelectedProductionMenuEntries[0]) + 1;
+		if (NextIndex < ProductionMenuEntries.Num())
+		{
+			ProductionMenuListView->SetSelection(ProductionMenuEntries[NextIndex]);
+		}
+	}
+
+	return FReply::Handled();
+}
+
 TSharedRef<SWidget> SProductionWizard::MakeMenuPanel()
 {
 	// Build the Revision Control menu entry, including its content panel
@@ -141,12 +232,14 @@ TSharedRef<SWidget> SProductionWizard::MakeMenuPanel()
 	// Construct the list view for the user menu entries
 	UserMenuListView = SNew(SListView<TSharedPtr<FProductionWizardMenuEntry>>)
 		.ListItemsSource(&UserMenuEntries)
+		.SelectionMode(ESelectionMode::Single)
 		.OnGenerateRow(this, &SProductionWizard::OnGenerateMenuRow)
 		.OnSelectionChanged(this, &SProductionWizard::OnUserMenuSelectionChanged);
 
 	// Construct the list view for the production menu entries
 	ProductionMenuListView = SNew(SListView<TSharedPtr<FProductionWizardMenuEntry>>)
 		.ListItemsSource(&ProductionMenuEntries)
+		.SelectionMode(ESelectionMode::Single)
 		.OnGenerateRow(this, &SProductionWizard::OnGenerateMenuRow)
 		.OnSelectionChanged(this, &SProductionWizard::OnProductionMenuSelectionChanged);
 
@@ -219,10 +312,8 @@ TSharedRef<ITableRow> SProductionWizard::OnGenerateMenuRow(TSharedPtr<FProductio
 
 void SProductionWizard::OnUserMenuSelectionChanged(TSharedPtr<FProductionWizardMenuEntry> SelectedEntry, ESelectInfo::Type SelectInfo)
 {
-	if (SelectInfo != ESelectInfo::Direct)
+	if (SelectedEntry)
 	{
-		check((SelectedEntry->Panel));
-
 		ProductionMenuListView->ClearSelection();
 		ContentPanelSlot->AttachWidget(SelectedEntry->Panel.ToSharedRef());
 	}
@@ -230,10 +321,8 @@ void SProductionWizard::OnUserMenuSelectionChanged(TSharedPtr<FProductionWizardM
 
 void SProductionWizard::OnProductionMenuSelectionChanged(TSharedPtr<FProductionWizardMenuEntry> SelectedEntry, ESelectInfo::Type SelectInfo)
 {
-	if (SelectInfo != ESelectInfo::Direct)
+	if (SelectedEntry)
 	{
-		check((SelectedEntry->Panel));
-
 		UserMenuListView->ClearSelection();
 		ContentPanelSlot->AttachWidget(SelectedEntry->Panel.ToSharedRef());
 	}
diff --git a/Engine/Plugins/VirtualProduction/Shots/Source/ShotsEditor/Private/UI/SProductionWizard.h b/Engine/Plugins/VirtualProduction/Shots/Source/ShotsEditor/Private/UI/SProductionWizard.h
--- a/Engine/Plugins/VirtualProduction/Shots/Source/ShotsEditor/Private/UI/SProductionWizard.h
+++ b/Engine/Plugins/VirtualProduction/Shots/Source/ShotsEditor/Private/UI/SProductionWizard.h
@@ -55,6 +55,12 @@ private:
 	/** Callback when one of the production menu entries is selected, which will cause its panel to be displayed on the right side of the production wizard */
 	void OnProductionMenuSelectionChanged(TSharedPtr<FProductionWizardMenuEntry> SelectedEntry, ESelectInfo::Type SelectInfo);
 
+	/** Cycles backwards through the available panels */
+	FReply OnBackClicked();
+
+	/** Cycles forwards through the available panels */
+	FReply OnNextClicked();
+
 private:
 	/** Slots for the menu and menu entry content panels */
 	SSplitter::FSlot* MenuPanelSlot = nullptr;
