Naming Tokens: Add feature to register namespace filters that should always be used when evaluating tokens.
Apply namespace deny list as a filter to naming tokens namespaces.

diff --git a/Engine/Plugins/Developer/NamingTokens/Source/NamingTokens/Private/NamingTokensEditorSubsystem.cpp b/Engine/Plugins/Developer/NamingTokens/Source/NamingTokens/Private/NamingTokensEditorSubsystem.cpp
--- a/Engine/Plugins/Developer/NamingTokens/Source/NamingTokens/Private/NamingTokensEditorSubsystem.cpp
+++ b/Engine/Plugins/Developer/NamingTokens/Source/NamingTokens/Private/NamingTokensEditorSubsystem.cpp
@@ -198,6 +198,14 @@ FNamingTokenResultData UNamingTokensEditorSubsystem::EvaluateTokenText(const FTe
 	Result.EvaluatedText = ProcessedTokenText; // Set so we always have something in case no tokens were evaluated.
 
 	TSet<FString> Namespaces = GetNamingTokenNamespacesFromString(InTokenText.ToString(), InFilter);
+
+	// These external filters are intentionally executed before adding any namespaces from the input FNamingTokenFilterArgs filter. 
+	// This allows the caller of this function to guarantee that the namespaces they want included will always be allowed, regardless of other filters that may execute.
+	for (const TPair<FName, FFilterNamespace>& DelegatePair : FilterNamespaceDelegates)
+	{
+		DelegatePair.Value.Execute(Namespaces);
+	}
+
 	Namespaces.Append(InFilter.NamespacesToInclude); // Tokens could have been written without the namespace but should still be scoped to the filter.
 
 	const FNamingTokensEvaluationData EvaluationData;
@@ -247,6 +255,16 @@ TArray<FString> UNamingTokensEditorSubsystem::GetGlobalNamespaces() const
 	return GlobalNamespaces.Array();
 }
 
+void UNamingTokensEditorSubsystem::RegisterNamespaceFilter(const FName OwnerName, FFilterNamespace Delegate)
+{
+	FilterNamespaceDelegates.Add(OwnerName, MoveTemp(Delegate));
+}
+
+void UNamingTokensEditorSubsystem::UnregisterNamespaceFilter(const FName OwnerName)
+{
+	FilterNamespaceDelegates.Remove(OwnerName);
+}
+
 TSet<FString> UNamingTokensEditorSubsystem::GetNamingTokenNamespacesFromString(const FString& InTokenString, const FNamingTokenFilterArgs& InFilter) const
 {
 	TSet<FString> Namespaces;
diff --git a/Engine/Plugins/Developer/NamingTokens/Source/NamingTokens/Public/NamingTokensEditorSubsystem.h b/Engine/Plugins/Developer/NamingTokens/Source/NamingTokens/Public/NamingTokensEditorSubsystem.h
--- a/Engine/Plugins/Developer/NamingTokens/Source/NamingTokens/Public/NamingTokensEditorSubsystem.h
+++ b/Engine/Plugins/Developer/NamingTokens/Source/NamingTokens/Public/NamingTokensEditorSubsystem.h
@@ -10,6 +10,8 @@
 
 class UNamingTokens;
 
+DECLARE_DELEGATE_OneParam(FFilterNamespace, TSet<FString>& /*Namespaces*/);
+
 USTRUCT(BlueprintType)
 struct FNamingTokenFilterArgs
 {
@@ -89,6 +91,12 @@ public:
 	UFUNCTION(BlueprintCallable, Category = "NamingTokens")
 	TArray<FString> GetAllNamespaces() const;
 
+	/** Register a delegate to filter the list of namespaces that can be used to evaluate tokens */
+	void RegisterNamespaceFilter(const FName OwnerName, FFilterNamespace Delegate);
+
+	/** Remove a delegate from the list of namespace filters */
+	void UnregisterNamespaceFilter(const FName OwnerName);
+
 	/** Locate all reference naming token namespaces from a given string. */
 	TSet<FString> GetNamingTokenNamespacesFromString(const FString& InTokenString, const FNamingTokenFilterArgs& InFilter = FNamingTokenFilterArgs()) const;
 	
@@ -109,4 +117,7 @@ private:
 
 	/** Namespaces considered global (don't need to include namespace to access). */
 	TSet<FString> GlobalNamespaces;
+
+	/** Filters that will execute just before evaluating a token string */
+	TMap<FName, FFilterNamespace> FilterNamespaceDelegates;
 };
diff --git a/Engine/Plugins/VirtualProduction/Shots/Source/ShotsEditor/Private/ProductionSettings.cpp b/Engine/Plugins/VirtualProduction/Shots/Source/ShotsEditor/Private/ProductionSettings.cpp
--- a/Engine/Plugins/VirtualProduction/Shots/Source/ShotsEditor/Private/ProductionSettings.cpp
+++ b/Engine/Plugins/VirtualProduction/Shots/Source/ShotsEditor/Private/ProductionSettings.cpp
@@ -2,11 +2,14 @@
 #include "ProductionSettings.h"
 
 #include "AssetToolsModule.h"
+#include "Editor.h"
 #include "HAL/FileManager.h"
 #include "HAL/PlatformFileManager.h"
 #include "LevelSequenceProjectSettings.h"
 #include "Misc/ConfigCacheIni.h"
+#include "Misc/CoreDelegates.h"
 #include "MovieSceneToolsProjectSettings.h"
+#include "NamingTokensEditorSubsystem.h"
 #include "Sections/MovieSceneSubSection.h"
 #include "SourceControlHelpers.h"
 
@@ -89,6 +92,15 @@ void UProductionSettings::PostInitProperties()
 	{
 		ApplyProjectOverrides();
 	}
+
+	// Register a filter with the naming tokens subsystem to apply the active production's namespace deny list when evaluating tokens
+	auto RegisterNamespaceFilter = [this]()
+		{
+			UNamingTokensEditorSubsystem* NamingTokensSubsystem = GEditor->GetEditorSubsystem<UNamingTokensEditorSubsystem>();
+			NamingTokensSubsystem->RegisterNamespaceFilter("CinematicProductionSettings", FFilterNamespace::CreateUObject(this, &UProductionSettings::FilterNamingTokenNamespaces));
+		};
+
+	FCoreDelegates::OnPostEngineInit.AddLambda(RegisterNamespaceFilter);
 }
 
 void UProductionSettings::ApplyProjectOverrides()
@@ -493,6 +505,17 @@ void UProductionSettings::RemoveNamespaceFromDenyList(FGuid ProductionID, const
 	}
 }
 
+void UProductionSettings::FilterNamingTokenNamespaces(TSet<FString>& Namespaces)
+{
+	if (ActiveProduction)
+	{
+		for (const FString& Namespace : ActiveProduction->NamingTokenNamespaceDenyList)
+		{
+			Namespaces.Remove(Namespace);
+		}
+	}
+}
+
 void UProductionSettings::AddAssetNaming(FGuid ProductionID, const UClass* AssetClass, const FString& DefaultName)
 {
 	if (!AssetClass)
diff --git a/Engine/Plugins/VirtualProduction/Shots/Source/ShotsEditor/Private/ProductionSettings.h b/Engine/Plugins/VirtualProduction/Shots/Source/ShotsEditor/Private/ProductionSettings.h
--- a/Engine/Plugins/VirtualProduction/Shots/Source/ShotsEditor/Private/ProductionSettings.h
+++ b/Engine/Plugins/VirtualProduction/Shots/Source/ShotsEditor/Private/ProductionSettings.h
@@ -170,6 +170,9 @@ private:
 	/** Overrides the DefaultAssetNames property of the asset tools project settings based on the active production setting */
 	void OverrideDefaultAssetNames();
 
+	/** Apply the active production's namespace deny list to the input set of namespace names */
+	void FilterNamingTokenNamespaces(TSet<FString>& Namespaces);
+
 	/** Set the active production name and write out the new value to EditorPerProjectUserSettings.ini */
 	void SetActiveProductionName();
 
