ShotManagement: Add Revision Control panel to Production Wizard
#rnx

diff --git a/Engine/Plugins/VirtualProduction/Shots/Source/ShotsEditor/Private/ShotsEditorStyle.cpp b/Engine/Plugins/VirtualProduction/Shots/Source/ShotsEditor/Private/ShotsEditorStyle.cpp
--- a/Engine/Plugins/VirtualProduction/Shots/Source/ShotsEditor/Private/ShotsEditorStyle.cpp
+++ b/Engine/Plugins/VirtualProduction/Shots/Source/ShotsEditor/Private/ShotsEditorStyle.cpp
@@ -34,6 +34,9 @@ FShotsEditorStyle::FShotsEditorStyle()
 	Set("Icons.Folder", new FSlateVectorImageBrush(RootToContentDir(TEXT("Slate/folder-closed.svg")), Icon16x16));
 	Set("Icons.Export", new FSlateVectorImageBrush(RootToContentDir(TEXT("Slate/ExportAll.svg")), Icon16x16));
 
+	Set("Badges.RevisionControlWarning", new FSlateVectorImageBrush(RootToCoreContentDir(TEXT("Starship/SourceControl/Status/RevisionControlBadgeWarning.svg")), Icon16x16, FStyleColors::Warning));
+	Set("Badges.RevisionControlConnected", new FSlateVectorImageBrush(RootToCoreContentDir(TEXT("Starship/SourceControl/Status/RevisionControlBadgeConnected.svg")), Icon16x16, FStyleColors::Success));
+
 	// Fonts
 	Set("ProductionWizard.TitleFont", DEFAULT_FONT("Bold", 10));
 	Set("ProductionWizard.HeadingFont", DEFAULT_FONT("Bold", 14));
diff --git a/Engine/Plugins/VirtualProduction/Shots/Source/ShotsEditor/Private/UI/SProductionWizard.cpp b/Engine/Plugins/VirtualProduction/Shots/Source/ShotsEditor/Private/UI/SProductionWizard.cpp
--- a/Engine/Plugins/VirtualProduction/Shots/Source/ShotsEditor/Private/UI/SProductionWizard.cpp
+++ b/Engine/Plugins/VirtualProduction/Shots/Source/ShotsEditor/Private/UI/SProductionWizard.cpp
@@ -7,6 +7,7 @@
 #include "UI/SNamingTokensPanel.h"
 #include "UI/SAssetNamingPanel.h"
 #include "UI/SProductionListPanel.h"
+#include "UI/SRevisionControlPanel.h"
 #include "UI/SSequencerSettingsPanel.h"
 #include "Widgets/Layout/SBorder.h"
 #include "Widgets/Layout/SSeparator.h"
@@ -189,7 +190,7 @@ TSharedRef<SWidget> SProductionWizard::MakeMenuPanel()
 	TSharedPtr<FProductionWizardMenuEntry> RevisionControlMenuEntry = MakeShared<FProductionWizardMenuEntry>();
 	RevisionControlMenuEntry->Label = LOCTEXT("RevisionControlMenuLabel", "Revision Control");
 	RevisionControlMenuEntry->Icon = FSlateIcon(FShotsEditorStyle::StyleName, "Icons.RevisionControl");
-	RevisionControlMenuEntry->Panel = MakeRevisionControlPanel();
+	RevisionControlMenuEntry->Panel = SNew(SRevisionControlPanel);
 
 	// Build the Productions menu entry, including its content panel
 	TSharedPtr<FProductionWizardMenuEntry> ProductionMenuEntry = MakeShared<FProductionWizardMenuEntry>();
@@ -328,26 +329,6 @@ void SProductionWizard::OnProductionMenuSelectionChanged(TSharedPtr<FProductionW
 	}
 }
 
-TSharedRef<SWidget> SProductionWizard::MakeRevisionControlPanel()
-{
-	// TODO: Fill in this stubbed out panel
-	return SNew(SVerticalBox)
-
-	+ SVerticalBox::Slot()
-		.FillHeight(1.0f)
-		.HAlign(HAlign_Center)
-		[
-			SNew(SHorizontalBox)
-
-			+ SHorizontalBox::Slot()
-				.FillWidth(1.0f)
-				.VAlign(VAlign_Center)
-				[
-					SNew(STextBlock).Text(LOCTEXT("RevisionControlPanel", "REVISION CONTROL PANEL"))
-				]
-		];
-}
-
 TSharedRef<SWidget> SProductionWizard::MakeFolderHierarchyPanel()
 {
 	// TODO: Fill in this stubbed out panel
diff --git a/Engine/Plugins/VirtualProduction/Shots/Source/ShotsEditor/Private/UI/SProductionWizard.h b/Engine/Plugins/VirtualProduction/Shots/Source/ShotsEditor/Private/UI/SProductionWizard.h
--- a/Engine/Plugins/VirtualProduction/Shots/Source/ShotsEditor/Private/UI/SProductionWizard.h
+++ b/Engine/Plugins/VirtualProduction/Shots/Source/ShotsEditor/Private/UI/SProductionWizard.h
@@ -43,7 +43,6 @@ private:
 	TSharedRef<SWidget> MakeButtonsPanel();
 
 	/** Functions for creating the panels for each of the production wizard's menu entries */
-	TSharedRef<SWidget> MakeRevisionControlPanel();
 	TSharedRef<SWidget> MakeFolderHierarchyPanel();
 
 	/** Generates a row in the menu list views */
diff --git a/Engine/Plugins/VirtualProduction/Shots/Source/ShotsEditor/Private/UI/SRevisionControlPanel.cpp b/Engine/Plugins/VirtualProduction/Shots/Source/ShotsEditor/Private/UI/SRevisionControlPanel.cpp
new file mode 100644
--- /dev/null
+++ b/Engine/Plugins/VirtualProduction/Shots/Source/ShotsEditor/Private/UI/SRevisionControlPanel.cpp
@@ -0,0 +1,137 @@
+// Copyright Epic Games, Inc. All Rights Reserved.
+
+#include "SRevisionControlPanel.h"
+
+#include "ISourceControlModule.h"
+#include "ISourceControlProvider.h"
+#include "ShotsEditorStyle.h"
+#include "Widgets/Images/SLayeredImage.h"
+#include "Widgets/Input/SButton.h"
+#include "Widgets/Layout/SBorder.h"
+#include "Widgets/SBoxPanel.h"
+#include "Widgets/Text/STextBlock.h"
+
+#define LOCTEXT_NAMESPACE "SRevisionControlPanel"
+
+void SRevisionControlPanel::Construct(const FArguments& InArgs)
+{
+	TSharedRef<SLayeredImage> SourceControlIcon =
+		SNew(SLayeredImage)
+		.ColorAndOpacity(FSlateColor::UseForeground())
+		.Image(FShotsEditorStyle::Get().GetBrush("Icons.RevisionControl"));
+
+	SourceControlIcon->AddLayer(TAttribute<const FSlateBrush*>::CreateStatic(&SRevisionControlPanel::GetSourceControlIconBadge));
+
+	ChildSlot
+		[
+			SNew(SBorder)
+				.BorderImage(FAppStyle::Get().GetBrush("Brushes.Panel"))
+				.Padding(16.0f)
+				[
+					SNew(SVerticalBox)
+
+					// Title
+					+ SVerticalBox::Slot()
+						.AutoHeight()
+						.Padding(0.0f, 0.0f, 0.0f, 4.0f)
+						[
+							SNew(STextBlock)
+								.Text(LOCTEXT("RevisionControlTitle", "User Setup"))
+								.Font(FShotsEditorStyle::Get().GetFontStyle("ProductionWizard.TitleFont"))
+						]
+
+					// Heading
+					+ SVerticalBox::Slot()
+						.AutoHeight()
+						.Padding(0.0f, 0.0f, 0.0f, 4.0f)
+						[
+							SNew(STextBlock)
+								.Text(LOCTEXT("RevisionControlHeading", "Revision Control"))
+								.Font(FShotsEditorStyle::Get().GetFontStyle("ProductionWizard.HeadingFont"))
+						]
+
+					// Info Text 1
+					+ SVerticalBox::Slot()
+						.AutoHeight()
+						.Padding(0.0f, 0.0f, 0.0f, 8.0f)
+						[
+							SNew(STextBlock)
+								.Text(LOCTEXT("RevisionControlInfoText1", "Revision control helps you collaborate with your team members and back up changes to your project."))
+								.AutoWrapText(true)
+						]
+
+					// Info Text 2
+					+ SVerticalBox::Slot()
+						.AutoHeight()
+						.Padding(0.0f, 0.0f, 0.0f, 16.0f)
+						[
+							SNew(STextBlock)
+								.Text(LOCTEXT("RevisionControlInfoText2", "If your project is in a revision control system, use the settings here to connect the Unreal editor to your project’s repository."))
+								.AutoWrapText(true)
+						]
+
+					// Connect to Source Control Button
+					+ SVerticalBox::Slot()
+						.AutoHeight()
+						.HAlign(HAlign_Left)
+						.Padding(0.0f, 0.0f, 0.0f, 0.0f)
+						[
+							SNew(SButton)
+								.ContentPadding(FMargin(2.0f))
+								.OnClicked_Lambda([]() -> FReply
+									{
+										// Opens the source control login dialog
+										ISourceControlModule& SourceControlModule = ISourceControlModule::Get();
+										SourceControlModule.ShowLoginDialog(FSourceControlLoginClosed(), ELoginWindowMode::Modal, EOnLoginWindowStartup::PreserveProvider);
+										return FReply::Handled();
+									})
+								[
+									SNew(SHorizontalBox)
+
+									+ SHorizontalBox::Slot()
+										.AutoWidth()
+										.Padding(0.0f, 0.0f, 4.0f, 0.0f)
+										[
+											SourceControlIcon
+										]
+
+									+ SHorizontalBox::Slot()
+										.AutoWidth()
+										[
+											SNew(STextBlock)
+												.Text_Lambda([]() -> FText
+													{
+														if (ISourceControlModule::Get().IsEnabled())
+														{
+															return LOCTEXT("ChangeRevisionControlSettings", "Change Revision Control Settings");
+														}
+														return LOCTEXT("ConnectToRevisionControl", "Connect to Revision Control");
+													})
+										]
+								]
+						]
+				]
+		];
+}
+
+const FSlateBrush* SRevisionControlPanel::GetSourceControlIconBadge()
+{
+	ISourceControlModule& SourceControlModule = ISourceControlModule::Get();
+	if (SourceControlModule.IsEnabled())
+	{
+		if (!SourceControlModule.GetProvider().IsAvailable())
+		{
+			static const FSlateBrush* WarningBrush = FSlateIcon(FShotsEditorStyle::StyleName, "Badges.RevisionControlWarning").GetIcon();
+			return WarningBrush;
+		}
+		else
+		{
+			static const FSlateBrush* ConnectedBrush = FSlateIcon(FShotsEditorStyle::StyleName, "Badges.RevisionControlConnected").GetIcon();
+			return ConnectedBrush;
+		}
+	}
+
+	return nullptr;
+}
+
+#undef LOCTEXT_NAMESPACE
diff --git a/Engine/Plugins/VirtualProduction/Shots/Source/ShotsEditor/Private/UI/SRevisionControlPanel.h b/Engine/Plugins/VirtualProduction/Shots/Source/ShotsEditor/Private/UI/SRevisionControlPanel.h
new file mode 100644
--- /dev/null
+++ b/Engine/Plugins/VirtualProduction/Shots/Source/ShotsEditor/Private/UI/SRevisionControlPanel.h
@@ -0,0 +1,24 @@
+// Copyright Epic Games, Inc. All Rights Reserved.
+
+#pragma once
+
+#include "Widgets/SCompoundWidget.h"
+
+#include "Widgets/DeclarativeSyntaxSupport.h"
+
+/**
+ * UI for the Revision Control panel in the Production Wizard
+ */
+class SRevisionControlPanel : public SCompoundWidget
+{
+public:
+	SLATE_BEGIN_ARGS(SRevisionControlPanel) {}
+	SLATE_END_ARGS()
+
+	void Construct(const FArguments& InArgs);
+
+private:
+
+	/** Returns the icon badge to indicate connection status */
+	static const FSlateBrush* GetSourceControlIconBadge();
+};
