ShotManagement: Fix production settings to properly restore movie scene tools project settings when the active production is reset to None

#rnx

diff --git a/Engine/Plugins/VirtualProduction/Shots/Source/ShotsEditor/Private/ProductionSettings.cpp b/Engine/Plugins/VirtualProduction/Shots/Source/ShotsEditor/Private/ProductionSettings.cpp
--- a/Engine/Plugins/VirtualProduction/Shots/Source/ShotsEditor/Private/ProductionSettings.cpp
+++ b/Engine/Plugins/VirtualProduction/Shots/Source/ShotsEditor/Private/ProductionSettings.cpp
@@ -54,22 +54,17 @@ void UProductionSettings::PostInitProperties()
 {
 	Super::PostInitProperties();
 
-	// Register a callback to detect when a ULevelSequenceProjectSettings property has changed to be notified when the DefaultDisplayRate changes
-	ULevelSequenceProjectSettings* LevelSequenceSettings = GetMutableDefault<ULevelSequenceProjectSettings>();
-	LevelSequenceSettings->OnSettingChanged().AddUObject(this, &UProductionSettings::OnLevelSequenceSettingChanged);
-
-	const UMovieSceneToolsProjectSettings* const MovieSceneToolsSettings = GetDefault<UMovieSceneToolsProjectSettings>();
+	CacheProjectDefaults();
 
-	// Cache the original tooltips for these properties to reset if the active production is set to "None"
+	// Cache the original values for these tooltips to reset if the active production is set to "None"
+	const ULevelSequenceProjectSettings* const LevelSequenceSettings = GetMutableDefault<ULevelSequenceProjectSettings>();
 	FProperty* DefaultDisplayRateProperty = LevelSequenceSettings->GetClass()->FindPropertyByName(GET_MEMBER_NAME_CHECKED(ULevelSequenceProjectSettings, DefaultDisplayRate));
 	OriginalDefaultDisplayRateTooltip = DefaultDisplayRateProperty->GetMetaData(TEXT("ToolTip"));
 
+	const UMovieSceneToolsProjectSettings* const MovieSceneToolsSettings = GetDefault<UMovieSceneToolsProjectSettings>();
 	FProperty* DefaultStartTimeProperty = MovieSceneToolsSettings->GetClass()->FindPropertyByName(GET_MEMBER_NAME_CHECKED(UMovieSceneToolsProjectSettings, DefaultStartTime));
 	OriginalDefaultStartTimeTooltip = DefaultStartTimeProperty->GetMetaData(TEXT("ToolTip"));
 
-	// Cache the original config value for the default display rate, so that it can be reset if the user sets their active production to "None"
-	DefaultDisplayRateProjectSetting = LevelSequenceSettings->DefaultDisplayRate;
-
 	// Load the serialized active production name from user settings
 	const FString ConfigSection = UProductionSettings::StaticClass()->GetPathName();
 	const FString PropertyName = GET_MEMBER_NAME_CHECKED(UProductionSettings, ActiveProductionName).ToString();
@@ -103,6 +98,16 @@ void UProductionSettings::PostInitProperties()
 	FCoreDelegates::OnPostEngineInit.AddLambda(RegisterNamespaceFilter);
 }
 
+void UProductionSettings::CacheProjectDefaults()
+{
+	// Cache the original values for these properties to reset if the active production is set to "None"
+	const ULevelSequenceProjectSettings* const LevelSequenceSettings = GetMutableDefault<ULevelSequenceProjectSettings>();
+	ProjectDefaultDisplayRate = LevelSequenceSettings->DefaultDisplayRate;
+
+	const UMovieSceneToolsProjectSettings* const MovieSceneToolsSettings = GetDefault<UMovieSceneToolsProjectSettings>();
+	ProjectDefaultStartTime = MovieSceneToolsSettings->DefaultStartTime;
+}
+
 void UProductionSettings::ApplyProjectOverrides()
 {
 	// Modify the property flags and tooltip of these overridden project settings
@@ -151,7 +156,7 @@ void UProductionSettings::OverrideDefaultDisplayRate()
 	}
 	else
 	{
-		LevelSequenceSettings->DefaultDisplayRate = DefaultDisplayRateProjectSetting;
+		LevelSequenceSettings->DefaultDisplayRate = ProjectDefaultDisplayRate;
 	}
 
 	// Update the default display rate cvar to match
@@ -163,13 +168,17 @@ void UProductionSettings::OverrideDefaultDisplayRate()
 
 void UProductionSettings::OverrideDefaultStartTime()
 {
+	// Update the DefaultStartTime property of the movie scene tools project settings to the value of the active production
+	// If there is no active production, the movie scene tools project setting is reset to its default config value
+	UMovieSceneToolsProjectSettings* MovieSceneToolsSettings = GetMutableDefault<UMovieSceneToolsProjectSettings>();
+
 	if (ActiveProduction)
 	{
-		// Update the DefaultStartTime property of the movie scene tools project settings to the value of the active production
-		UMovieSceneToolsProjectSettings* MovieSceneToolsSettings = GetMutableDefault<UMovieSceneToolsProjectSettings>();
-
 		MovieSceneToolsSettings->DefaultStartTime = ActiveProduction->DefaultDisplayRate.AsSeconds(FFrameNumber(ActiveProduction->DefaultStartFrame));
-		MovieSceneToolsSettings->SaveConfig();
+	}
+	else
+	{
+		MovieSceneToolsSettings->DefaultStartTime = ProjectDefaultStartTime;
 	}
 }
 
@@ -271,15 +280,6 @@ void UProductionSettings::PostEditChangeProperty(struct FPropertyChangedEvent& P
 	}
 }
 
-void UProductionSettings::OnLevelSequenceSettingChanged(UObject* Object, struct FPropertyChangedEvent& PropertyChangedEvent)
-{
-	if (PropertyChangedEvent.GetPropertyName() == GET_MEMBER_NAME_CHECKED(ULevelSequenceProjectSettings, DefaultDisplayRate))
-	{
-		const ULevelSequenceProjectSettings* const LevelSequenceSettings = GetDefault<ULevelSequenceProjectSettings>();
-		DefaultDisplayRateProjectSetting = LevelSequenceSettings->DefaultDisplayRate;
-	}
-}
-
 void UProductionSettings::AddProduction()
 {
 	const int32 ArrayIndex = Productions.Add(FCinematicProduction());
@@ -351,6 +351,12 @@ FCinematicProduction* UProductionSettings::GetActiveProduction() const
 
 void UProductionSettings::SetActiveProduction(FGuid ProductionID)
 {
+	// If the active production is currently None, before setting it and applying the project overrides, cache some project defaults to restore later if the active production is set to None again
+	if (!ActiveProduction)
+	{
+		CacheProjectDefaults();
+	}
+
 	if (ActiveProduction)
 	{
 		// Unregister any default asset names that the previously active production had set
diff --git a/Engine/Plugins/VirtualProduction/Shots/Source/ShotsEditor/Private/ProductionSettings.h b/Engine/Plugins/VirtualProduction/Shots/Source/ShotsEditor/Private/ProductionSettings.h
--- a/Engine/Plugins/VirtualProduction/Shots/Source/ShotsEditor/Private/ProductionSettings.h
+++ b/Engine/Plugins/VirtualProduction/Shots/Source/ShotsEditor/Private/ProductionSettings.h
@@ -152,9 +152,6 @@ public:
 	FOnActiveProductionChanged& OnActiveProductionChanged() { return ActiveProductionChangedDelegate; }
 
 private:
-	/** Callback when a property of ULevelSequenceProjectSettings has changed */
-	void OnLevelSequenceSettingChanged(UObject* Object, struct FPropertyChangedEvent& PropertyChangedEvent);
-
 	/** Applies overrides to various project settings based on the active production settings */
 	void ApplyProjectOverrides();
 
@@ -176,6 +173,9 @@ private:
 	/** Set the active production name and write out the new value to EditorPerProjectUserSettings.ini */
 	void SetActiveProductionName();
 
+	/** Cache the default project settings that are overridden by the active production, used to reset when there is no active production */
+	void CacheProjectDefaults();
+
 public:
 	/** Name of the active production */
 	UPROPERTY(VisibleAnywhere, Category = "Default", meta=(DisplayName="Active Production"))
@@ -189,8 +189,9 @@ private:
 	/** Pointer to the active production (in the Productions array) */
 	FCinematicProduction* ActiveProduction;
 
-	/** Default value of DefaultDisplayRate in the level sequence project setting, used to reset it when there is no active production */
-	FString DefaultDisplayRateProjectSetting;
+	/** Cached default project settings that are overridden by the active production, used to reset when there is no active production */
+	FString ProjectDefaultDisplayRate;
+	float ProjectDefaultStartTime;
 
 	/** Default asset names previously registered with AssetTools, used to reset them when the active production changes */
 	TMap<const UClass*, FString> ProjectDefaultAssetNames;
