ShotManagement: Add option in the production wizard to duplicate a production

#rnx

diff --git a/Engine/Plugins/VirtualProduction/Shots/Source/ShotsEditor/Private/ProductionSettings.cpp b/Engine/Plugins/VirtualProduction/Shots/Source/ShotsEditor/Private/ProductionSettings.cpp
--- a/Engine/Plugins/VirtualProduction/Shots/Source/ShotsEditor/Private/ProductionSettings.cpp
+++ b/Engine/Plugins/VirtualProduction/Shots/Source/ShotsEditor/Private/ProductionSettings.cpp
@@ -292,6 +292,28 @@ void UProductionSettings::AddProduction()
 	ProductionListChangedDelegate.Broadcast();
 }
 
+void UProductionSettings::DuplicateProduction(FGuid ProductionID)
+{
+	const FCinematicProduction* ProductionToCopy = Productions.FindByPredicate([ProductionID](const FCinematicProduction& Other) { return Other.ProductionID == ProductionID; });
+	if (ProductionToCopy)
+	{
+		FCinematicProduction DuplicateProduction = *ProductionToCopy;
+
+		// Give the duplicate production its own GUID and a unique name
+		DuplicateProduction.ProductionID = FGuid::NewGuid();
+		DuplicateProduction.ProductionName = GetUniqueProductionName(ProductionToCopy->ProductionName);
+
+		Productions.Add(DuplicateProduction);
+
+		// Reset the active production pointer, in case the array was reallocated
+		ActiveProduction = Productions.FindByPredicate([this](const FCinematicProduction& Other) { return Other.ProductionName == ActiveProductionName; });
+
+		TryUpdateDefaultConfigFile();
+
+		ProductionListChangedDelegate.Broadcast();
+	}
+}
+
 void UProductionSettings::DeleteProduction(FGuid ProductionID)
 {
 	if (!ProductionID.IsValid())
@@ -620,14 +642,41 @@ void UProductionSettings::SetActiveProductionName()
 
 FString UProductionSettings::GetUniqueProductionName() const
 {
-	const FText NameTemplate = LOCTEXT("NewProductionText", "New Production");
-	FString WorkingName = NameTemplate.ToString();
-	int32 Suffix = 1;
+	const FText BaseName = LOCTEXT("NewProductionBaseName", "NewProduction");
+	return GetUniqueProductionName(BaseName.ToString());
+}
+
+FString UProductionSettings::GetUniqueProductionName(const FString& BaseName) const
+{
+	// This implementation is based on a similar utility in AssetTools for creating a unique asset name.
+	// If the input BaseName does not end in a numeric character, then a 1 will be appended to it.
+	// If it already ends in a numeric character, then that number will get incremented.
+	// Ex: NewProduction -> NewProduction1
+	// Ex: NewProduction1 -> NewProduction2
+
+	// Find the index in the string of the last non-numeric character
+	int32 CharIndex = BaseName.Len() - 1;
+	while (CharIndex >= 0 && BaseName[CharIndex] >= TEXT('0') && BaseName[CharIndex] <= TEXT('9'))
+	{
+		--CharIndex;
+	}
+
+	// Trim the numeric characters off the end of the BaseName string, but remember the integer that was trimmed off to increment and append to the output
+	int32 IntSuffix = 1;
+	FString TrimmedBaseName = BaseName;
+	if (CharIndex >= 0 && CharIndex < BaseName.Len() - 1)
+	{
+		TrimmedBaseName = BaseName.Left(CharIndex + 1);
+
+		const FString TrailingInteger = BaseName.RightChop(CharIndex + 1);
+		IntSuffix = FCString::Atoi(*TrailingInteger) + 1;
+	}
 
+	FString WorkingName = TrimmedBaseName;
 	while (Productions.ContainsByPredicate([WorkingName](const FCinematicProduction& Other) { return Other.ProductionName == WorkingName; }))
 	{
-		WorkingName = FString::Printf(TEXT("%s %d"), *NameTemplate.ToString(), Suffix);
-		Suffix++;
+		WorkingName = FString::Printf(TEXT("%s%d"), *TrimmedBaseName, IntSuffix);
+		IntSuffix++;
 	}
 
 	return WorkingName;
diff --git a/Engine/Plugins/VirtualProduction/Shots/Source/ShotsEditor/Private/ProductionSettings.h b/Engine/Plugins/VirtualProduction/Shots/Source/ShotsEditor/Private/ProductionSettings.h
--- a/Engine/Plugins/VirtualProduction/Shots/Source/ShotsEditor/Private/ProductionSettings.h
+++ b/Engine/Plugins/VirtualProduction/Shots/Source/ShotsEditor/Private/ProductionSettings.h
@@ -85,6 +85,9 @@ public:
 	/** Adds a new production to the list */
 	void AddProduction();
 
+	/** Adds a duplicate of the input production to the list */
+	void DuplicateProduction(FGuid ProductionID);
+
 	/** Removes the production matching the input ID from the list */
 	void DeleteProduction(FGuid ProductionID);
 
@@ -144,6 +147,7 @@ public:
 
 	/** Returns a new unique production name */
 	FString GetUniqueProductionName() const;
+	FString GetUniqueProductionName(const FString& BaseName) const;
 
 	/** Returns the delegate that broadcasts when a production is added/removed */
 	FOnProductionListChanged& OnProductionListChanged() { return ProductionListChangedDelegate; }
diff --git a/Engine/Plugins/VirtualProduction/Shots/Source/ShotsEditor/Private/UI/SProductionListPanel.cpp b/Engine/Plugins/VirtualProduction/Shots/Source/ShotsEditor/Private/UI/SProductionListPanel.cpp
--- a/Engine/Plugins/VirtualProduction/Shots/Source/ShotsEditor/Private/UI/SProductionListPanel.cpp
+++ b/Engine/Plugins/VirtualProduction/Shots/Source/ShotsEditor/Private/UI/SProductionListPanel.cpp
@@ -380,6 +380,20 @@ FReply SProductionListButton::OnMouseButtonDown(const FGeometry& MyGeometry, con
 			EUserInterfaceActionType::Button
 		);
 
+		MenuBuilder.AddMenuEntry(
+			LOCTEXT("DuplicateProduction", "Duplicate"),
+			LOCTEXT("DuplicateProductionToolTip", "Duplicate production"),
+			FSlateIcon(FAppStyle::GetAppStyleSetName(), "Icons.Duplicate"),
+			FUIAction(FExecuteAction::CreateLambda([this]()
+				{
+					UProductionSettings* ProductionSettings = GetMutableDefault<UProductionSettings>();
+					ProductionSettings->DuplicateProduction(ProductionID.Get());
+				}
+			)),
+			NAME_None,
+			EUserInterfaceActionType::Button
+		);
+
 		FWidgetPath WidgetPath = MouseEvent.GetEventPath() != nullptr ? *MouseEvent.GetEventPath() : FWidgetPath();
 
 		FSlateApplication::Get().PushMenu(
