ShotManagement: Add import/export support for production settings from the production wizard
#rnx

diff --git a/Engine/Plugins/VirtualProduction/Shots/Source/ShotsEditor/Private/ProductionSettings.cpp b/Engine/Plugins/VirtualProduction/Shots/Source/ShotsEditor/Private/ProductionSettings.cpp
--- a/Engine/Plugins/VirtualProduction/Shots/Source/ShotsEditor/Private/ProductionSettings.cpp
+++ b/Engine/Plugins/VirtualProduction/Shots/Source/ShotsEditor/Private/ProductionSettings.cpp
@@ -292,6 +292,18 @@ void UProductionSettings::AddProduction()
 	ProductionListChangedDelegate.Broadcast();
 }
 
+void UProductionSettings::AddProduction(const FCinematicProduction& ProductionToAdd)
+{
+	Productions.Add(ProductionToAdd);
+
+	// Reset the active production pointer, in case the array was reallocated
+	ActiveProduction = Productions.FindByPredicate([this](const FCinematicProduction& Other) { return Other.ProductionName == ActiveProductionName; });
+
+	TryUpdateDefaultConfigFile();
+
+	ProductionListChangedDelegate.Broadcast();
+}
+
 void UProductionSettings::DuplicateProduction(FGuid ProductionID)
 {
 	const FCinematicProduction* ProductionToCopy = Productions.FindByPredicate([ProductionID](const FCinematicProduction& Other) { return Other.ProductionID == ProductionID; });
diff --git a/Engine/Plugins/VirtualProduction/Shots/Source/ShotsEditor/Private/ProductionSettings.h b/Engine/Plugins/VirtualProduction/Shots/Source/ShotsEditor/Private/ProductionSettings.h
--- a/Engine/Plugins/VirtualProduction/Shots/Source/ShotsEditor/Private/ProductionSettings.h
+++ b/Engine/Plugins/VirtualProduction/Shots/Source/ShotsEditor/Private/ProductionSettings.h
@@ -82,9 +82,12 @@ public:
 	virtual void PostEditChangeProperty(struct FPropertyChangedEvent& PropertyChangedEvent) override;
 	//~ End UObject overrides
 
-	/** Adds a new production to the list */
+	/** Adds a new empty production to the list */
 	void AddProduction();
 
+	/** Adds the input production to the list */
+	void AddProduction(const FCinematicProduction& ProductionToAdd);
+
 	/** Adds a duplicate of the input production to the list */
 	void DuplicateProduction(FGuid ProductionID);
 
diff --git a/Engine/Plugins/VirtualProduction/Shots/Source/ShotsEditor/Private/UI/SProductionListPanel.cpp b/Engine/Plugins/VirtualProduction/Shots/Source/ShotsEditor/Private/UI/SProductionListPanel.cpp
--- a/Engine/Plugins/VirtualProduction/Shots/Source/ShotsEditor/Private/UI/SProductionListPanel.cpp
+++ b/Engine/Plugins/VirtualProduction/Shots/Source/ShotsEditor/Private/UI/SProductionListPanel.cpp
@@ -2,9 +2,12 @@
 
 #include "SProductionListPanel.h"
 
+#include "DesktopPlatformModule.h"
 #include "Framework/Application/MenuStack.h"
 #include "Framework/Application/SlateApplication.h"
 #include "Framework/MultiBox/MultiBoxBuilder.h"
+#include "HAL/FileManager.h"
+#include "JsonObjectConverter.h"
 #include "Layout/WidgetPath.h"
 #include "ProductionSettings.h"
 #include "ShotsEditorStyle.h"
@@ -107,9 +110,9 @@ void SProductionListPanel::Construct(const FArguments& InArgs)
 								+ SHorizontalBox::Slot()
 								.AutoWidth()
 								[
-									// TODO: Add support for importing a production when this button is clicked
 									SNew(SButton)
 										.ContentPadding(FMargin(2.0f))
+										.OnClicked(this, &SProductionListPanel::ImportProduction)
 										[
 											SNew(SHorizontalBox)
 
@@ -219,8 +222,8 @@ TSharedRef<ITableRow> SProductionListPanel::OnGenerateProductionRow(TSharedPtr<F
 				.AutoWidth()
 				.VAlign(VAlign_Center)
 				[
-					// TODO: Add support for exporting a production when this button is clicked
 					SNew(SButton)
+						.OnClicked(this, &SProductionListPanel::ExportProduction, InItem)
 						.ButtonStyle(FShotsEditorStyle::Get(), "ProductionWizard.RecessedButton")
 						[
 							SNew(SImage).Image(FShotsEditorStyle::Get().GetBrush("Icons.Export"))
@@ -241,6 +244,101 @@ TSharedRef<ITableRow> SProductionListPanel::OnGenerateProductionRow(TSharedPtr<F
 		];
 }
 
+FReply SProductionListPanel::ImportProduction()
+{
+	// Prompt the user to choose a production setting .json file to open
+	IDesktopPlatform* DesktopPlatform = FDesktopPlatformModule::Get();
+	const void* ParentWindowHandle = FSlateApplication::Get().FindBestParentWindowHandleForDialogs(nullptr);
+	const FString Title = TEXT("Import Production Settings");
+	const FString DefaultPath = FPaths::ProjectSavedDir();
+	const FString DefaultFile = TEXT("");
+	const FString FileTypes = TEXT("json|*.json");
+	const uint32 FileFlags = 0;
+
+	TArray<FString> OpenFileNames;
+	const bool bFileSelected = DesktopPlatform->OpenFileDialog(ParentWindowHandle, Title, DefaultPath, DefaultFile, FileTypes, FileFlags, OpenFileNames);
+
+	if (!bFileSelected || OpenFileNames.Num() != 1)
+	{
+		return FReply::Handled();
+	}
+
+	const FString JsonFileName = OpenFileNames[0];
+
+	if (TUniquePtr<FArchive> FileReader = TUniquePtr<FArchive>(IFileManager::Get().CreateFileReader(*JsonFileName)))
+	{
+		TSharedRef<TJsonReader<>> JsonReader = TJsonReaderFactory<>::Create(FileReader.Get());
+
+		TSharedPtr<FJsonObject> JsonObject = MakeShared<FJsonObject>();
+		if (FJsonSerializer::Deserialize(JsonReader, JsonObject))
+		{
+			TSharedPtr<FCinematicProduction> ImportedProduction = MakeShared<FCinematicProduction>();
+
+			// We enforce strict mode to ensure that every field in the UStruct is present in the imported json.
+			constexpr int64 CheckFlags = 0;
+			constexpr int64 SkipFlags = 0;
+			constexpr bool bStrictMode = true;
+			if (FJsonObjectConverter::JsonObjectToUStruct<FCinematicProduction>(JsonObject.ToSharedRef(), ImportedProduction.Get(), CheckFlags, SkipFlags, bStrictMode))
+			{
+				// If the import and conversion from json was successful, add the imported production to the Production Project Settings list of productions
+				UProductionSettings* ProductionSettings = GetMutableDefault<UProductionSettings>();
+
+				// Give the imported production a new unique ID and ensure that its name is unique. 
+				ImportedProduction->ProductionID = FGuid::NewGuid();
+				ImportedProduction->ProductionName = ProductionSettings->GetUniqueProductionName(ImportedProduction->ProductionName);
+
+				ProductionSettings->AddProduction(*ImportedProduction);
+			}
+		}
+	}
+
+	return FReply::Handled();
+}
+
+FReply SProductionListPanel::ExportProduction(TSharedPtr<FCinematicProduction> InItem)
+{
+	UProductionSettings* ProductionSettings = GetMutableDefault<UProductionSettings>();
+	const FCinematicProduction* ProductionToExport = ProductionSettings->GetProduction(InItem->ProductionID);
+	if (!ProductionToExport)
+	{
+		return FReply::Handled();
+	}
+
+	// Convert the production settings to a json object that can be written to .json file on disk
+	if (const TSharedPtr<FJsonObject>& JsonObject = FJsonObjectConverter::UStructToJsonObject<FCinematicProduction>(*ProductionToExport))
+	{
+		// Prompt the user to choose a location to save the production setting .json file
+		IDesktopPlatform* DesktopPlatform = FDesktopPlatformModule::Get();
+		const void* ParentWindowHandle = FSlateApplication::Get().FindBestParentWindowHandleForDialogs(nullptr);
+		const FString Title = TEXT("Export Production Settings");
+		const FString DefaultPath = FPaths::ProjectSavedDir();
+		const FString DefaultFile = ProductionToExport->ProductionName + ".json";
+		const FString FileTypes = TEXT("json|*.json");
+		const uint32 FileFlags = 0;
+
+		TArray<FString> SaveFileNames;
+		const bool bFileSelected = DesktopPlatform->SaveFileDialog(ParentWindowHandle, Title, DefaultPath, DefaultFile, FileTypes, FileFlags, SaveFileNames);
+
+		if (!bFileSelected || SaveFileNames.Num() != 1)
+		{
+			return FReply::Handled();
+		}
+
+		const FString JsonFileName = SaveFileNames[0];
+
+		// Write the contents of the json object to disk
+		if (TUniquePtr<FArchive> FileWriter = TUniquePtr<FArchive>(IFileManager::Get().CreateFileWriter(*JsonFileName)))
+		{
+			TSharedRef<TJsonWriter<>> JsonWriter = TJsonWriterFactory<>::Create(FileWriter.Get());
+
+			FJsonSerializer::Serialize(JsonObject.ToSharedRef(), JsonWriter);
+			FileWriter->Close();
+		}
+	}
+
+	return FReply::Handled();
+}
+
 SProductionListButton::SProductionListButton()
 	: ProductionID(*this)
 {
diff --git a/Engine/Plugins/VirtualProduction/Shots/Source/ShotsEditor/Private/UI/SProductionListPanel.h b/Engine/Plugins/VirtualProduction/Shots/Source/ShotsEditor/Private/UI/SProductionListPanel.h
--- a/Engine/Plugins/VirtualProduction/Shots/Source/ShotsEditor/Private/UI/SProductionListPanel.h
+++ b/Engine/Plugins/VirtualProduction/Shots/Source/ShotsEditor/Private/UI/SProductionListPanel.h
@@ -73,6 +73,12 @@ private:
 	/** Callback when the production list view is finished regenerating its rows */
 	void OnProductionListItemsRebuilt();
 
+	/** Exports the input production to a .json file on disk containing all of its production settings */
+	FReply ExportProduction(TSharedPtr<FCinematicProduction> InItem);
+
+	/** Imports a production setting .json file from disk and adds it to the Production Project Settings list of productions */
+	FReply ImportProduction();
+
 private:
 	/** List items sources for the production list view */
 	TArray<TSharedPtr<FCinematicProduction>> ProductionList;
diff --git a/Engine/Plugins/VirtualProduction/Shots/Source/ShotsEditor/ShotsEditor.build.cs b/Engine/Plugins/VirtualProduction/Shots/Source/ShotsEditor/ShotsEditor.build.cs
--- a/Engine/Plugins/VirtualProduction/Shots/Source/ShotsEditor/ShotsEditor.build.cs
+++ b/Engine/Plugins/VirtualProduction/Shots/Source/ShotsEditor/ShotsEditor.build.cs
@@ -20,9 +20,12 @@ namespace UnrealBuildTool.Rules
 					"ClassViewer",
 					"Core",
 					"CoreUObject",
+					"DesktopPlatform",
 					"DeveloperSettings",
 					"Engine",
 					"InputCore",
+					"Json",
+					"JsonUtilities",
 					"LevelSequence",
 					"LevelSequenceEditor",
 					"MovieScene",
