ShotManagement: Adds an array of folder paths to the Production Settings that represents a template folder hierarchy for that production. Also adds the UI for the Folder Hierarchy panel to the Production Wizard.

From the Folder Hierarchy panel, users can add/rename/remove folders from the active production's template, and then bulk create folders in the content browser based on the folder template.

#rnx

diff --git a/Engine/Plugins/Developer/DirectoryPlaceholder/Source/DirectoryPlaceholder/Private/DirectoryPlaceholder.h b/Engine/Plugins/Developer/DirectoryPlaceholder/Source/DirectoryPlaceholder/Private/DirectoryPlaceholder.h
deleted file mode 100644
diff --git a/Engine/Plugins/Developer/DirectoryPlaceholder/Source/DirectoryPlaceholder/Private/DirectoryPlaceholderUtils.cpp b/Engine/Plugins/Developer/DirectoryPlaceholder/Source/DirectoryPlaceholder/Private/DirectoryPlaceholderUtils.cpp
--- a/Engine/Plugins/Developer/DirectoryPlaceholder/Source/DirectoryPlaceholder/Private/DirectoryPlaceholderUtils.cpp
+++ b/Engine/Plugins/Developer/DirectoryPlaceholder/Source/DirectoryPlaceholder/Private/DirectoryPlaceholderUtils.cpp
@@ -12,7 +12,7 @@
 
 namespace UE::DirectoryPlaceholder::Private
 {
-	/** Recurisvely delete all unnecessary placeholder assets in this folder (and sub-folders) */
+	/** Recursively delete all unnecessary placeholder assets in this folder (and sub-folders) */
 	bool CleanupPlaceholdersInternal(FName Path)
 	{
 		FAssetRegistryModule& AssetRegistryModule = FModuleManager::Get().GetModuleChecked<FAssetRegistryModule>(TEXT("AssetRegistry"));
@@ -88,4 +88,22 @@ void UDirectoryPlaceholderLibrary::CleanupPlaceholdersInPaths(const TArray<FStri
 	}
 }
 
+void UDirectoryPlaceholderLibrary::DeletePlaceholdersInPath(const FString& Path)
+{
+	FAssetRegistryModule& AssetRegistryModule = FModuleManager::Get().GetModuleChecked<FAssetRegistryModule>(TEXT("AssetRegistry"));
+
+	// Get all of the assets in the current folder
+	TArray<FAssetData> AssetDataList;
+	constexpr bool bRecursive = true;
+	AssetRegistryModule.Get().GetAssetsByPath(*Path, AssetDataList, bRecursive);
+
+	TArray<FAssetData> PlacesholdersToDelete = AssetDataList.FilterByPredicate([](const FAssetData& AssetData) {return AssetData.GetClass() == UDirectoryPlaceholder::StaticClass(); });
+
+	if (PlacesholdersToDelete.Num() > 0)
+	{
+		const bool bShowConfirmation = false;
+		const int32 NumAssetsDeleted = ObjectTools::DeleteAssets(PlacesholdersToDelete, bShowConfirmation);
+	}
+}
+
 #undef LOCTEXT_NAMESPACE
diff --git a/Engine/Plugins/Developer/DirectoryPlaceholder/Source/DirectoryPlaceholder/Private/DirectoryPlaceholderUtils.h b/Engine/Plugins/Developer/DirectoryPlaceholder/Source/DirectoryPlaceholder/Private/DirectoryPlaceholderUtils.h
deleted file mode 100644
diff --git a/Engine/Plugins/Developer/DirectoryPlaceholder/Source/DirectoryPlaceholder/Public/DirectoryPlaceholder.h b/Engine/Plugins/Developer/DirectoryPlaceholder/Source/DirectoryPlaceholder/Public/DirectoryPlaceholder.h
new file mode 100644
--- /dev/null
+++ b/Engine/Plugins/Developer/DirectoryPlaceholder/Source/DirectoryPlaceholder/Public/DirectoryPlaceholder.h
@@ -0,0 +1,35 @@
+// Copyright Epic Games, Inc. All Rights Reserved.
+
+#pragma once
+
+#include "ContentBrowserFrontEndFilterExtension.h"
+
+#include "DirectoryPlaceholder.generated.h"
+
+class FFrontendFilter;
+class FFrontendFilterCategory;
+
+/** 
+ * Extremely lightweight object that can be used as a placeholder in an otherwise empty directory.
+ * The presence of a placeholder object allows that directory to be submitted to source control.
+ */
+UCLASS()
+class DIRECTORYPLACEHOLDER_API UDirectoryPlaceholder : public UObject
+{
+	GENERATED_BODY()
+
+public:
+	UDirectoryPlaceholder() = default;
+};
+
+/** Content Browser filter used to show/hide directory placeholder assets */
+UCLASS()
+class UDirectoryPlaceholderSearchFilter : public UContentBrowserFrontEndFilterExtension
+{
+	GENERATED_BODY()
+
+protected:
+	// Begin UContentBrowserFrontEndFilterExtension Interface
+	virtual void AddFrontEndFilterExtensions(TSharedPtr<FFrontendFilterCategory> DefaultCategory, TArray<TSharedRef<FFrontendFilter>>& InOutFilterList) const override;
+	// End UContentBrowserFrontEndFilterExtension Interface
+};
diff --git a/Engine/Plugins/Developer/DirectoryPlaceholder/Source/DirectoryPlaceholder/Public/DirectoryPlaceholderUtils.h b/Engine/Plugins/Developer/DirectoryPlaceholder/Source/DirectoryPlaceholder/Public/DirectoryPlaceholderUtils.h
new file mode 100644
--- /dev/null
+++ b/Engine/Plugins/Developer/DirectoryPlaceholder/Source/DirectoryPlaceholder/Public/DirectoryPlaceholderUtils.h
@@ -0,0 +1,29 @@
+// Copyright Epic Games, Inc. All Rights Reserved.
+
+#pragma once
+
+#include "Kismet/BlueprintFunctionLibrary.h"
+
+#include "DirectoryPlaceholderUtils.generated.h"
+
+/**
+ * Library functions for operations on directory placeholder assets
+ */
+UCLASS()
+class DIRECTORYPLACEHOLDER_API UDirectoryPlaceholderLibrary : public UBlueprintFunctionLibrary
+{
+	GENERATED_BODY()
+
+public:
+	/** Delete all unnecessary placeholder assets in this folder (and sub-folders) */
+	UFUNCTION(BlueprintCallable, Category = "DirectoryPlaceholder")
+	static void CleanupPlaceholdersInPath(const FString& Path);
+
+	/** Delete all unnecessary placeholder assets in these folders (and sub-folders) */
+	UFUNCTION(BlueprintCallable, Category = "DirectoryPlaceholder")
+	static void CleanupPlaceholdersInPaths(const TArray<FString>& Paths);
+	
+	/** Delete all placeholder assets in this folder (and sub-folders) */
+	UFUNCTION(BlueprintCallable, Category = "DirectoryPlaceholder")
+	static void DeletePlaceholdersInPath(const FString& Path);
+};
diff --git a/Engine/Plugins/VirtualProduction/Shots/Source/ShotsEditor/Private/ProductionSettings.cpp b/Engine/Plugins/VirtualProduction/Shots/Source/ShotsEditor/Private/ProductionSettings.cpp
--- a/Engine/Plugins/VirtualProduction/Shots/Source/ShotsEditor/Private/ProductionSettings.cpp
+++ b/Engine/Plugins/VirtualProduction/Shots/Source/ShotsEditor/Private/ProductionSettings.cpp
@@ -344,11 +344,7 @@ void UProductionSettings::DeleteProduction(FGuid ProductionID)
 
 		if (bIsActiveProduction)
 		{
-			ActiveProduction = nullptr;
-			SetActiveProductionName();
-
-			// This will restore the level sequence project setting to its original default setting (before any production overrode it)
-			OverrideDefaultDisplayRate();
+			SetActiveProduction(ProductionID);
 		}
 
 		TryUpdateDefaultConfigFile();
@@ -617,6 +613,46 @@ void UProductionSettings::RemoveAssetNaming(FGuid ProductionID, const UClass* As
 	}
 }
 
+void UProductionSettings::AddTemplateFolder(FGuid ProductionID, const FString& Path, bool bCreateIfMissing)
+{
+	FCinematicProduction* Production = Productions.FindByPredicate([ProductionID](const FCinematicProduction& Other) { return Other.ProductionID == ProductionID; });
+	if (Production)
+	{
+		FFolderTemplate NewTemplate;
+		NewTemplate.InternalPath = Path;
+		NewTemplate.bCreateIfMissing = bCreateIfMissing;
+
+		Production->TemplateFolders.Add(MoveTemp(NewTemplate));
+		TryUpdateDefaultConfigFile();
+	}
+}
+
+void UProductionSettings::RemoveTemplateFolder(FGuid ProductionID, const FString& Path)
+{
+	FCinematicProduction* Production = Productions.FindByPredicate([ProductionID](const FCinematicProduction& Other) { return Other.ProductionID == ProductionID; });
+	if (Production)
+	{
+		int32 IndexToDelete = Production->TemplateFolders.IndexOfByPredicate([Path](const FFolderTemplate& Other) { return Other.InternalPath == Path; });
+		if (Production->TemplateFolders.IsValidIndex(IndexToDelete))
+		{
+			Production->TemplateFolders.RemoveAt(IndexToDelete);
+		}
+
+		TryUpdateDefaultConfigFile();
+	}
+}
+
+void UProductionSettings::SetTemplateFolderHierarchy(FGuid ProductionID, const TArray<FFolderTemplate>& TemplateHierarchy)
+{
+	FCinematicProduction* Production = Productions.FindByPredicate([ProductionID](const FCinematicProduction& Other) { return Other.ProductionID == ProductionID; });
+	if (Production)
+	{
+		Production->TemplateFolders = TemplateHierarchy;
+
+		TryUpdateDefaultConfigFile();
+	}
+}
+
 void UProductionSettings::SetActiveDisplayRate(FFrameRate DisplayRate)
 {
 	if (ActiveProduction)
diff --git a/Engine/Plugins/VirtualProduction/Shots/Source/ShotsEditor/Private/ProductionSettings.h b/Engine/Plugins/VirtualProduction/Shots/Source/ShotsEditor/Private/ProductionSettings.h
--- a/Engine/Plugins/VirtualProduction/Shots/Source/ShotsEditor/Private/ProductionSettings.h
+++ b/Engine/Plugins/VirtualProduction/Shots/Source/ShotsEditor/Private/ProductionSettings.h
@@ -20,6 +20,21 @@ enum class ESubsequencePriority : uint8
 	BottomUp
 };
 
+/**
+ * Properties of a folder in the production's template folder hierarchy
+ */
+USTRUCT()
+struct FFolderTemplate
+{
+	GENERATED_BODY()
+
+	UPROPERTY(config, EditAnywhere, Category = "Default")
+	FString InternalPath;
+
+	UPROPERTY(config, EditAnywhere, Category = "Default")
+	bool bCreateIfMissing = true;
+};
+
 /**
  * Collection of production settings to override the project/editor behavior
  */
@@ -58,6 +73,10 @@ public:
 	/** List of default names for specific asset types */
 	UPROPERTY(config, VisibleAnywhere, BlueprintReadOnly, Category = "Default")
 	TMap<TObjectPtr<const UClass>, FString> DefaultAssetNames;
+
+	/** List of folder paths that represent a template folder hierarchy to be used for this production */
+	UPROPERTY(config, EditAnywhere, Category = "Default")
+	TArray<FFolderTemplate> TemplateFolders;
 };
 
 /**
@@ -139,6 +158,15 @@ public:
 	/** Removes an entry from the DefaultAssetNames map of the production matching the input ID */
 	void RemoveAssetNaming(FGuid ProductionID, const UClass* AssetClass);
 
+	/** Add a new path to the input production's list of template folders */
+	void AddTemplateFolder(FGuid ProductionID, const FString& Path, bool bCreateIfMissing=true);
+
+	/** Removes a path from the input production's list of template folders */
+	void RemoveTemplateFolder(FGuid ProductionID, const FString& Path);
+
+	/** Sets the input production's template folder hierarchy to the input array of template folders */
+	void SetTemplateFolderHierarchy(FGuid ProductionID, const TArray<FFolderTemplate>& TemplateHierarchy);
+
 	/** Sets the DefaultDisplayRate of the active production */
 	void SetActiveDisplayRate(FFrameRate DisplayRate);
 
diff --git a/Engine/Plugins/VirtualProduction/Shots/Source/ShotsEditor/Private/ShotsEditorStyle.cpp b/Engine/Plugins/VirtualProduction/Shots/Source/ShotsEditor/Private/ShotsEditorStyle.cpp
--- a/Engine/Plugins/VirtualProduction/Shots/Source/ShotsEditor/Private/ShotsEditorStyle.cpp
+++ b/Engine/Plugins/VirtualProduction/Shots/Source/ShotsEditor/Private/ShotsEditorStyle.cpp
@@ -33,10 +33,14 @@ FShotsEditorStyle::FShotsEditorStyle()
 	Set("Icons.AssetNaming", new FSlateVectorImageBrush(RootToContentDir(TEXT("Slate/Rename.svg")), Icon16x16));
 	Set("Icons.Folder", new FSlateVectorImageBrush(RootToContentDir(TEXT("Slate/folder-closed.svg")), Icon16x16));
 	Set("Icons.Export", new FSlateVectorImageBrush(RootToContentDir(TEXT("Slate/ExportAll.svg")), Icon16x16));
+	Set("Icons.Initialize", new FSlateVectorImageBrush(RootToContentDir(TEXT("Slate/Initialize_16.svg")), Icon16x16));
 
 	Set("Badges.RevisionControlWarning", new FSlateVectorImageBrush(RootToCoreContentDir(TEXT("Starship/SourceControl/Status/RevisionControlBadgeWarning.svg")), Icon16x16, FStyleColors::Warning));
 	Set("Badges.RevisionControlConnected", new FSlateVectorImageBrush(RootToCoreContentDir(TEXT("Starship/SourceControl/Status/RevisionControlBadgeConnected.svg")), Icon16x16, FStyleColors::Success));
 
+	Set("Badges.FolderNew", new FSlateVectorImageBrush(RootToContentDir(TEXT("Slate/CircleBadgeDot.svg")), Icon16x16, FStyleColors::Success));
+	Set("Badges.FolderRename", new FSlateVectorImageBrush(RootToContentDir(TEXT("Slate/CircleBadgeDot.svg")), Icon16x16, FStyleColors::Warning));
+
 	// Fonts
 	Set("ProductionWizard.TitleFont", DEFAULT_FONT("Bold", 10));
 	Set("ProductionWizard.HeadingFont", DEFAULT_FONT("Bold", 14));
diff --git a/Engine/Plugins/VirtualProduction/Shots/Source/ShotsEditor/Private/UI/SFolderHierarchyPanel.cpp b/Engine/Plugins/VirtualProduction/Shots/Source/ShotsEditor/Private/UI/SFolderHierarchyPanel.cpp
new file mode 100644
--- /dev/null
+++ b/Engine/Plugins/VirtualProduction/Shots/Source/ShotsEditor/Private/UI/SFolderHierarchyPanel.cpp
@@ -0,0 +1,932 @@
+// Copyright Epic Games, Inc. All Rights Reserved.
+
+#include "SFolderHierarchyPanel.h"
+
+#include "AssetRegistry/AssetRegistryModule.h"
+#include "AssetViewUtils.h"
+#include "ContentBrowserDataSubsystem.h"
+#include "ContentBrowserItem.h"
+#include "DirectoryPlaceholder.h"
+#include "DirectoryPlaceholderUtils.h"
+#include "Framework/Application/SlateApplication.h"
+#include "Framework/MultiBox/MultiBoxBuilder.h"
+#include "IContentBrowserDataModule.h"
+#include "HAL/FileManager.h"
+#include "Misc/MessageDialog.h"
+#include "ProductionSettings.h"
+#include "ShotsEditorStyle.h"
+#include "Styling/StyleColors.h"
+#include "UI/SActiveProductionCombo.h"
+#include "Widgets/Images/SLayeredImage.h"
+#include "Widgets/Input/SButton.h"
+#include "Widgets/Layout/SBorder.h"
+#include "Widgets/Layout/SSeparator.h"
+#include "Widgets/SBoxPanel.h"
+#include "Widgets/Text/SInlineEditableTextBlock.h"
+
+#define LOCTEXT_NAMESPACE "SFolderHierarchyPanel"
+
+FName SFolderHierarchyPanel::RootPath = TEXT("/Game");
+
+void SFolderHierarchyPanel::Construct(const FArguments& InArgs)
+{
+	// Subscribe to be notified when the Production Settings active productions has changed
+	UProductionSettings* ProductionSettings = GetMutableDefault<UProductionSettings>();
+	ActiveProductionChangedHandle = ProductionSettings->OnActiveProductionChanged().AddSP(this, &SFolderHierarchyPanel::BuildTreeFromProductionTemplate);
+
+	// Create the tree root
+	RootItem = MakeShared<FTemplateFolderTreeItem>();
+	RootItem->Path.SetPathFromString(RootPath.ToString(), EContentBrowserPathType::Internal);
+	RootItem->Status = ETemplateFolderStatus::Exists;
+
+	// Initialize the tree view using the active production's folder template
+	BuildTreeFromProductionTemplate();
+
+	TreeView = SNew(STreeView<TSharedPtr<FTemplateFolderTreeItem>>)
+		.TreeItemsSource(&FolderItemsSource)
+		.SelectionMode(ESelectionMode::Single)
+		.OnGenerateRow(this, &SFolderHierarchyPanel::OnGenerateRow)
+		.OnGetChildren(this, &SFolderHierarchyPanel::OnGetChildren)
+		.OnItemsRebuilt(this, &SFolderHierarchyPanel::OnItemsRebuilt)
+		.OnKeyDownHandler(this, &SFolderHierarchyPanel::OnKeyDown)
+		.OnSetExpansionRecursive(this, &SFolderHierarchyPanel::SetExpansionRecursive)
+		.OnContextMenuOpening(this, &SFolderHierarchyPanel::OnContextMenuOpening);
+
+	// Start with the entire template tree expanded
+	SetExpansionRecursive(RootItem, true);
+
+	ChildSlot
+		[
+			SNew(SVerticalBox)
+
+			// Active Production Selector
+			+ SVerticalBox::Slot()
+				.AutoHeight()
+				[
+					SNew(SActiveProductionCombo)
+				]
+
+			// Separator
+			+ SVerticalBox::Slot()
+				.AutoHeight()
+				[
+					SNew(SSeparator)
+						.Orientation(Orient_Horizontal)
+						.Thickness(2.0f)
+				]
+
+			// Folder Hierarchy Panel
+			+ SVerticalBox::Slot()
+				.FillHeight(1.0f)
+				[
+					SNew(SBorder)
+						.BorderImage(FAppStyle::Get().GetBrush("Brushes.Panel"))
+						.Padding(16.0f)
+						[
+							SNew(SVerticalBox)
+
+							// Title
+							+ SVerticalBox::Slot()
+								.AutoHeight()
+								.Padding(0.0f, 0.0f, 0.0f, 4.0f)
+								[
+									SNew(STextBlock)
+										.Text(LOCTEXT("FolderHierarchyTitle", "Production Settings"))
+										.Font(FShotsEditorStyle::Get().GetFontStyle("ProductionWizard.TitleFont"))
+								]
+
+							// Heading
+							+ SVerticalBox::Slot()
+								.AutoHeight()
+								.Padding(0.0f, 0.0f, 0.0f, 4.0f)
+								[
+									SNew(STextBlock)
+										.Text(LOCTEXT("FolderHierarchyHeading", "Folder Hierarchy"))
+										.Font(FShotsEditorStyle::Get().GetFontStyle("ProductionWizard.HeadingFont"))
+								]
+
+							// Info Text 1
+							+ SVerticalBox::Slot()
+								.AutoHeight()
+								.Padding(0.0f, 0.0f, 0.0f, 8.0f)
+								[
+									SNew(STextBlock).Text(LOCTEXT("FolderHierarchyInfoText1", "Set up the folders you want to use for your production’s assets.\n"
+											"(Folders that already contain assets can only be deleted in the Content Browser.)"))
+								]
+
+							// Info Text 2
+							+ SVerticalBox::Slot()
+								.AutoHeight()
+								.Padding(0.0f, 0.0f, 0.0f, 16.0f)
+								[
+									SNew(STextBlock).Text(LOCTEXT("FolderHierarchyInfoText2", "Click Create Template Folders when you are done."))
+								]
+
+							// Main Panel Content
+							+ SVerticalBox::Slot()
+								[
+									// This entire box will be disabled when there is no active production
+									SNew(SVerticalBox)
+										.IsEnabled_Lambda([]() -> bool
+											{
+												const UProductionSettings* const ProductionSettings = GetDefault<UProductionSettings>();
+												return (ProductionSettings->GetActiveProduction() != nullptr);
+											}
+										)
+
+									// Add / Reset Buttons
+									+ SVerticalBox::Slot()
+										.AutoHeight()
+										.Padding(0.0f, 0.0f, 0.0f, 8.0f)
+										[
+											SNew(SHorizontalBox)
+
+											// Add Folder To Template Button
+											+ SHorizontalBox::Slot()
+												.AutoWidth()
+												[
+													SNew(SButton)
+														.ContentPadding(FMargin(2.0f))
+														.OnClicked(this, &SFolderHierarchyPanel::OnAddFolderToTemplate)
+														[
+															SNew(SHorizontalBox)
+
+															+ SHorizontalBox::Slot()
+																.AutoWidth()
+																.Padding(0.0f, 0.0f, 4.0f, 0.0f)
+																[
+																	SNew(SImage)
+																		.Image(FAppStyle::Get().GetBrush("Icons.Plus"))
+																		.ColorAndOpacity(FStyleColors::AccentGreen)
+																]
+
+															+ SHorizontalBox::Slot()
+																.AutoWidth()
+																[
+																	SNew(STextBlock).Text(LOCTEXT("AddFolderToTemplateButton", "Add Folder To Template"))
+																]
+														]
+												]
+
+											// Reset Template Changes Button
+											+ SHorizontalBox::Slot()
+												.AutoWidth()
+												[
+													SNew(SButton)
+														.ContentPadding(FMargin(2.0f))
+														.OnClicked(this, &SFolderHierarchyPanel::OnResetTemplateChanges)
+														[
+															SNew(SHorizontalBox)
+
+															+ SHorizontalBox::Slot()
+																.AutoWidth()
+																.Padding(0.0f, 0.0f, 4.0f, 0.0f)
+																[
+																	SNew(SImage).Image(FAppStyle::Get().GetBrush("PropertyWindow.DiffersFromDefault"))
+																]
+
+															+ SHorizontalBox::Slot()
+																.AutoWidth()
+																[
+																	SNew(STextBlock).Text(LOCTEXT("OnResetTemplateChangesButton", "Reset Template Changes"))
+																]
+														]
+												]
+										]
+
+									// Tree View
+									+ SVerticalBox::Slot()
+										.FillHeight(1.0f)
+										.Padding(0.0f, 0.0f, 0.0f, 8.0f)
+										[
+											TreeView.ToSharedRef()
+										]
+
+									// Apply Changes Button
+									+ SVerticalBox::Slot()
+										.AutoHeight()
+										.HAlign(HAlign_Right)
+										.Padding(0.0f, 0.0f, 0.0f, 8.0f)
+										[
+											SNew(SButton)
+												.ContentPadding(FMargin(2.0f))
+												.OnClicked(this, &SFolderHierarchyPanel::OnCreateTemplateFolders)
+												[
+													SNew(SHorizontalBox)
+
+													+ SHorizontalBox::Slot()
+														.AutoWidth()
+														.Padding(0.0f, 0.0f, 4.0f, 0.0f)
+														[
+															SNew(SImage).Image(FShotsEditorStyle::Get().GetBrush("Icons.Initialize"))
+														]
+
+													+ SHorizontalBox::Slot()
+														.AutoWidth()
+														[
+															SNew(STextBlock).Text(LOCTEXT("CreateTemplateFoldersButton", "Create Template Folders"))
+														]
+												]
+										]
+								]
+						]
+				]
+		];
+}
+
+SFolderHierarchyPanel::~SFolderHierarchyPanel()
+{
+	if (UObjectInitialized())
+	{
+		if (UProductionSettings* ProductionSettings = GetMutableDefault<UProductionSettings>())
+		{
+			ProductionSettings->OnActiveProductionChanged().Remove(ActiveProductionChangedHandle);
+		}
+	}
+}
+
+TSharedRef<ITableRow> SFolderHierarchyPanel::OnGenerateRow(TSharedPtr<FTemplateFolderTreeItem> TreeItem, const TSharedRef<STableViewBase>& OwnerTable)
+{
+	TSharedRef<SLayeredImage> FolderIcon = SNew(SLayeredImage)
+		.ColorAndOpacity(FAppStyle::Get().GetSlateColor("ContentBrowser.DefaultFolderColor"))
+		.Image(FShotsEditorStyle::Get().GetBrush("Icons.Folder"));
+
+	FolderIcon->AddLayer(TAttribute<const FSlateBrush*>::CreateSP(this, &SFolderHierarchyPanel::GetFolderIconBadge, TreeItem));
+
+	return SNew(STableRow<TSharedPtr<FTemplateFolderTreeItem>>, OwnerTable)
+		.ShowSelection(true)
+		.Padding(FMargin(8.0f, 2.0f, 8.0f, 0.0f))
+		.Content()
+		[
+			SNew(SHorizontalBox)
+
+			// Folder Icon
+			+ SHorizontalBox::Slot()
+				.AutoWidth()
+				.VAlign(VAlign_Center)
+				.Padding(0.0f, 0.0f, 2.0f, 0.0f)
+				[
+					FolderIcon
+				]
+
+			// Folder Name TextBlock
+			+ SHorizontalBox::Slot()
+				.VAlign(VAlign_Center)
+				[
+					SAssignNew(TreeItem->NameWidget, SInlineEditableTextBlock)
+						.Text_Lambda([TreeItem]() -> FText
+							{
+								// The Content Browser overrides the display name of the top-level project content folder to "Content", which we want to match
+								if (TreeItem->Path.GetInternalPathName() == RootPath)
+								{
+									return LOCTEXT("ContentFolderDisplayNameOverride", "Content");
+								}
+								return FText::FromString(FPaths::GetPathLeaf(TreeItem->Path.GetInternalPathString()));
+							})
+						.IsReadOnly_Lambda([TreeItem]() -> bool
+							{
+								return (TreeItem->Path.GetInternalPathName() == RootPath);
+							})
+						.OnVerifyTextChanged(this, &SFolderHierarchyPanel::IsValidFolderName, TreeItem)
+						.OnTextCommitted(this, &SFolderHierarchyPanel::SetTemplateFolderName, TreeItem)
+				]
+
+			// Folder Status Text
+			+ SHorizontalBox::Slot()
+				.HAlign(HAlign_Right)
+				.VAlign(VAlign_Center)
+				[
+					SNew(STextBlock)
+						.Text(this, &SFolderHierarchyPanel::GetFolderStatusText, TreeItem)
+						.ColorAndOpacity(FLinearColor(1.0f, 1.0f, 1.0f, 0.4f))
+				]
+		];
+}
+
+void SFolderHierarchyPanel::OnGetChildren(TSharedPtr<FTemplateFolderTreeItem> TreeItem, TArray<TSharedPtr<FTemplateFolderTreeItem>>& OutNodes)
+{
+	OutNodes = TreeItem->Children;
+}
+
+void SFolderHierarchyPanel::OnItemsRebuilt()
+{
+	// Upon regenerating the tree view, allow the user to immediately interact with the name widget of the newly added template folder in order to rename it
+	if (MostRecentlyAddedItem)
+	{
+		FSlateApplication::Get().SetKeyboardFocus(MostRecentlyAddedItem->NameWidget.ToSharedRef());
+		MostRecentlyAddedItem->NameWidget->EnterEditingMode();
+		MostRecentlyAddedItem.Reset();
+	}
+}
+
+void SFolderHierarchyPanel::SetExpansionRecursive(TSharedPtr<FTemplateFolderTreeItem> TreeItem, bool bInExpand) const
+{
+	TreeView->SetItemExpansion(TreeItem, bInExpand);
+
+	for (const TSharedPtr<FTemplateFolderTreeItem>& ChildItem : TreeItem->Children)
+	{
+		SetExpansionRecursive(ChildItem, bInExpand);
+	}
+}
+
+TSharedPtr<SWidget> SFolderHierarchyPanel::OnContextMenuOpening()
+{
+	// The TreeView uses single selection mode, so at most one item can ever be selected by the user
+	TArray<TSharedPtr<FTemplateFolderTreeItem>> SelectedNodes = TreeView->GetSelectedItems();
+	if (SelectedNodes.Num() == 1)
+	{
+		constexpr bool bCloseAfterMenuSelection = true;
+		FMenuBuilder MenuBuilder(bCloseAfterMenuSelection, nullptr);
+
+		TSharedPtr<FTemplateFolderTreeItem>& SelectedTreeItem = SelectedNodes[0];
+
+		if (SelectedTreeItem->Path.GetInternalPathName() != RootPath)
+		{
+			// If the selected template folder does not currently exist in the Content Browser, give the user the option to toggle its status between "Create" and "Do Not Create"
+			if (SelectedTreeItem->Status == ETemplateFolderStatus::MissingCreate)
+			{
+				MenuBuilder.AddMenuEntry(
+					LOCTEXT("MarkAsDoNotCreate", "Mark as Do Not Create"),
+					LOCTEXT("MarkAsDoNotCreateToolTip", "Mark this folder so that it will not be created when clicking Create Template Folders"),
+					FSlateIcon(FAppStyle::GetAppStyleSetName(), "SourceControl.StatusIcon.Off"),
+					FUIAction(FExecuteAction::CreateSP(this, &SFolderHierarchyPanel::SetTemplateFolderStatusRecursive, SelectedTreeItem, ETemplateFolderStatus::MissingDoNotCreate)),
+					NAME_None,
+					EUserInterfaceActionType::Button
+				);
+			}
+			else if (SelectedTreeItem->Status == ETemplateFolderStatus::MissingDoNotCreate)
+			{
+				MenuBuilder.AddMenuEntry(
+					LOCTEXT("MarkAsCreate", "Mark as Create"),
+					LOCTEXT("MarkAsCreateToolTip", "Mark this folder so that it will not be created when clicking Create Template Folders"),
+					FSlateIcon(FAppStyle::GetAppStyleSetName(), "Icons.AddCircle"),
+					FUIAction(FExecuteAction::CreateSP(this, &SFolderHierarchyPanel::SetTemplateFolderStatusRecursive, SelectedTreeItem, ETemplateFolderStatus::MissingCreate)),
+					NAME_None,
+					EUserInterfaceActionType::Button
+				);
+			}
+		}
+
+		// Option to add a new child folder
+		MenuBuilder.AddMenuEntry(
+			LOCTEXT("AddChildFolderAction", "Add Child Folder"),
+			LOCTEXT("AddChildFolderTooltip", "Add child folder"),
+			FSlateIcon(FAppStyle::GetAppStyleSetName(), "Icons.Plus"),
+			FUIAction(FExecuteAction::CreateLambda([this]() { OnAddFolderToTemplate(); })),
+			NAME_None,
+			EUserInterfaceActionType::Button
+		);
+
+		if (SelectedTreeItem->Path.GetInternalPathName() != RootPath)
+		{
+			// Option to rename the template folder (and update the path of the children below it)
+			MenuBuilder.AddMenuEntry(
+				LOCTEXT("RenameAction", "Rename"),
+				LOCTEXT("RenameActionToolTip", "Rename template folder"),
+				FSlateIcon(FShotsEditorStyle::StyleName, "Icons.AssetNaming"),
+				FUIAction(
+					FExecuteAction::CreateLambda([SelectedTreeItem]()
+						{
+							FSlateApplication::Get().SetKeyboardFocus(SelectedTreeItem->NameWidget.ToSharedRef());
+							SelectedTreeItem->NameWidget->EnterEditingMode();
+						})),
+				NAME_None,
+				EUserInterfaceActionType::Button
+			);
+
+			// Option to delete the template folder and all of its children (recursively)
+			MenuBuilder.AddMenuEntry(
+				LOCTEXT("DeleteAction", "Delete"),
+				LOCTEXT("DeleteActionToolTip", "Delete this folder"),
+				FSlateIcon(FAppStyle::GetAppStyleSetName(), "Icons.Delete"),
+				FUIAction(FExecuteAction::CreateSP(this, &SFolderHierarchyPanel::DeleteTemplateFolder, SelectedTreeItem)),
+				NAME_None,
+				EUserInterfaceActionType::Button
+			);
+		}
+
+		return MenuBuilder.MakeWidget();
+	}
+
+	return SNullWidget::NullWidget;
+}
+
+FReply SFolderHierarchyPanel::OnKeyDown(const FGeometry& MyGeometry, const FKeyEvent& InKeyEvent)
+{
+	// The TreeView uses single selection mode, so at most one item can ever be selected by the user
+	TArray<TSharedPtr<FTemplateFolderTreeItem>> SelectedNodes = TreeView->GetSelectedItems();
+	if (SelectedNodes.Num() == 1)
+	{
+		TSharedPtr<FTemplateFolderTreeItem>& SelectedTreeItem = SelectedNodes[0];
+
+		if (InKeyEvent.GetKey() == EKeys::Delete)
+		{
+			DeleteTemplateFolder(SelectedTreeItem);
+		}
+	}
+
+	return FReply::Handled();
+}
+
+void SFolderHierarchyPanel::BuildTreeFromProductionTemplate()
+{
+	FolderItemsSource.Empty();
+	RootItem->Children.Empty();
+
+	FolderItemsSource.Add(RootItem);
+
+	const UProductionSettings* const ProductionSettings = GetMutableDefault<UProductionSettings>();
+	const FCinematicProduction* ActiveProduction = ProductionSettings->GetActiveProduction();
+	if (ActiveProduction)
+	{
+		FAssetRegistryModule& AssetRegistryModule = FModuleManager::GetModuleChecked<FAssetRegistryModule>("AssetRegistry");
+
+		// Sort the folder paths from the active production's template alphabetically to put them in the order they should appear in the tree
+		TArray<FFolderTemplate> ProductionTemplate = ActiveProduction->TemplateFolders;
+		CachedInitialState = ProductionTemplate;
+
+		ProductionTemplate.Sort([](const FFolderTemplate& A, const FFolderTemplate& B) { return A.InternalPath < B.InternalPath; });
+
+		for (const FFolderTemplate& FolderTemplate : ProductionTemplate)
+		{
+			const FString ParentPath = FPaths::GetPath(FolderTemplate.InternalPath);
+
+			// Walk the tree until we find an item whose path matches the parent path. The new tree item will be created as one of its children
+			// If no parent is found, we do not add the template folder to the tree view because it is considered ill-formed
+			TSharedPtr<FTemplateFolderTreeItem> ParentItem = FindItemAtPathRecursive(RootItem, ParentPath);
+			if (ParentItem)
+			{
+				TSharedPtr<FTemplateFolderTreeItem> NewItem = MakeShared<FTemplateFolderTreeItem>();
+				NewItem->Path.SetPathFromString(FolderTemplate.InternalPath, EContentBrowserPathType::Internal);
+
+				if (AssetRegistryModule.Get().PathExists(NewItem->Path.GetInternalPathString()))
+				{
+					NewItem->Status = ETemplateFolderStatus::Exists;
+				}
+				else
+				{
+					NewItem->Status = FolderTemplate.bCreateIfMissing ? ETemplateFolderStatus::MissingCreate : ETemplateFolderStatus::MissingDoNotCreate;
+				}
+
+				ParentItem->Children.Add(NewItem);
+				NewItem->Parent = ParentItem;
+			}
+		}
+	}
+
+	if (TreeView)
+	{
+		TreeView->RequestTreeRefresh();
+		SetExpansionRecursive(RootItem, true);
+	}
+}
+
+FReply SFolderHierarchyPanel::OnAddFolderToTemplate()
+{
+	FAssetRegistryModule& AssetRegistryModule = FModuleManager::GetModuleChecked<FAssetRegistryModule>("AssetRegistry");
+
+	// Get the parent item for the new template folder being added (this can be the root folder if no parent is currently selected)
+	// The TreeView uses single selection mode, so at most one item can ever be selected by the user
+	TSharedPtr<FTemplateFolderTreeItem> SelectedTreeItem = RootItem;
+
+	TArray<TSharedPtr<FTemplateFolderTreeItem>> SelectedNodes = TreeView->GetSelectedItems();
+	if (SelectedNodes.Num() == 1)
+	{
+		SelectedTreeItem = SelectedNodes[0];
+	}
+
+	TSharedPtr<FTemplateFolderTreeItem> NewFolder = MakeShared<FTemplateFolderTreeItem>();
+
+	const FString ParentPath = SelectedTreeItem->Path.GetInternalPathString();
+	const FString NewFolderName = CreateUniqueFolderName(SelectedTreeItem->Path);
+	NewFolder->Path.SetPathFromString(ParentPath / NewFolderName, EContentBrowserPathType::Internal);
+
+	if (AssetRegistryModule.Get().PathExists(NewFolder->Path.GetInternalPathString()))
+	{
+		NewFolder->Status = ETemplateFolderStatus::Exists;
+	}
+	else if (SelectedTreeItem->Status == ETemplateFolderStatus::MissingDoNotCreate)
+	{
+		NewFolder->Status = ETemplateFolderStatus::MissingDoNotCreate;
+	}
+	else
+	{
+		NewFolder->Status = ETemplateFolderStatus::MissingCreate;
+	}
+
+	SelectedTreeItem->Children.Add(NewFolder);
+	NewFolder->Parent = SelectedTreeItem;
+
+	MostRecentlyAddedItem = NewFolder;
+
+	// Sort the children alphabetically to maintain a good ordering with the new folder
+	SelectedTreeItem->Children.Sort([](const TSharedPtr<FTemplateFolderTreeItem>& A, const TSharedPtr<FTemplateFolderTreeItem>& B) { return A->Path.GetInternalPathString() < B->Path.GetInternalPathString(); });
+
+	// Add the new folder to the active production's template
+	UProductionSettings* ProductionSettings = GetMutableDefault<UProductionSettings>();
+	const FCinematicProduction* ActiveProduction = ProductionSettings->GetActiveProduction();
+	check(ActiveProduction);
+
+	const bool bCreateIfMissing = (NewFolder->Status == ETemplateFolderStatus::MissingDoNotCreate) ? false : true;
+	ProductionSettings->AddTemplateFolder(ActiveProduction->ProductionID, NewFolder->Path.GetInternalPathString(), bCreateIfMissing);
+
+	TreeView->SetItemExpansion(SelectedTreeItem, true);
+	TreeView->RequestTreeRefresh();
+
+	return FReply::Handled();
+}
+
+FReply SFolderHierarchyPanel::OnCreateTemplateFolders()
+{
+	// Walk the tree and create all template folders in the Content Browser that are marked as MissingCreate
+	CreateFolderFromTemplateRecursive(RootItem);
+
+	const UProductionSettings* const ProductionSettings = GetMutableDefault<UProductionSettings>();
+	const FCinematicProduction* ActiveProduction = ProductionSettings->GetActiveProduction();
+	check(ActiveProduction);
+	
+	CachedInitialState = ActiveProduction->TemplateFolders;
+
+	TreeView->RequestTreeRefresh();
+	SetExpansionRecursive(RootItem, true);
+
+	return FReply::Handled();
+}
+
+FReply SFolderHierarchyPanel::OnResetTemplateChanges()
+{
+	// Reset the active production's template to the latest cached state, then rebuild the tree view
+	UProductionSettings* ProductionSettings = GetMutableDefault<UProductionSettings>();
+	const FCinematicProduction* ActiveProduction = ProductionSettings->GetActiveProduction();
+	check(ActiveProduction);
+
+	ProductionSettings->SetTemplateFolderHierarchy(ActiveProduction->ProductionID, CachedInitialState);
+
+	BuildTreeFromProductionTemplate();
+
+	return FReply::Handled();
+}
+
+bool SFolderHierarchyPanel::IsValidFolderName(const FText& InText, FText& OutErrorMessage, TSharedPtr<FTemplateFolderTreeItem> TreeItem)
+{
+	// Check for empty text string
+	if (InText.IsEmpty())
+	{
+		OutErrorMessage = LOCTEXT("EmptyNameErrorMessage", "Please provide a name for this folder");
+		return false;
+	}
+
+	// Check for invalid characters
+	if (!AssetViewUtils::IsValidFolderName(InText.ToString(), OutErrorMessage))
+	{
+		return false;
+	}
+
+	// Check for duplicate folder names
+	const FString ParentPath = FPaths::GetPath(TreeItem->Path.GetInternalPathString());
+	TSharedPtr<FTemplateFolderTreeItem> FoundItem = FindItemAtPathRecursive(RootItem, ParentPath / InText.ToString());
+
+	if (FoundItem && FoundItem != TreeItem)
+	{
+		OutErrorMessage = LOCTEXT("DuplicateNameErrorMessage", "A folder already exists at this location with this name");
+		return false;
+	}
+
+	return true;
+}
+
+void SFolderHierarchyPanel::SetTemplateFolderName(const FText& InText, ETextCommit::Type InCommitType, TSharedPtr<FTemplateFolderTreeItem> TreeItem)
+{
+	// Early-out if the folder name has not actually changed
+	const FString OldPath = TreeItem->Path.GetInternalPathString();
+	const FString OldFolderName = FPaths::GetPathLeaf(OldPath);
+	if (OldFolderName.Equals(InText.ToString()))
+	{
+		return;
+	}
+
+	FAssetRegistryModule& AssetRegistryModule = FModuleManager::GetModuleChecked<FAssetRegistryModule>("AssetRegistry");
+
+	const FString NewPath = FPaths::GetPath(OldPath) / InText.ToString();
+
+	// If the folder being renamed exists in the content browser, and it is empty, prompt the user with a pop-up dialog asking if they want to rename the Content Browser folder
+	if (TreeItem->Status == ETemplateFolderStatus::Exists && IsFolderEmpty(TreeItem))
+	{
+		const FText DialogMessage = LOCTEXT("RenameDialogMessage", "An empty folder matching this template path already exists in the Content Browser. Do you want to rename that empty folder?");
+		EAppReturnType::Type Response = FMessageDialog::Open(EAppMsgType::YesNoCancel, DialogMessage);
+
+		// Rename the existing Content Browser folder
+		if (Response == EAppReturnType::Yes)
+		{
+			const FString NewRelativeFilePath = FPackageName::LongPackageNameToFilename(NewPath);
+			const FString NewAbsoluteFilePath = FPaths::ConvertRelativePathToFull(NewRelativeFilePath);
+
+			// Create the directory on disk, if it doesn't already exist
+			if (!IFileManager::Get().DirectoryExists(*NewAbsoluteFilePath))
+			{
+				constexpr bool bCreateParentFoldersIfMissing = false;
+				IFileManager::Get().MakeDirectory(*NewAbsoluteFilePath, bCreateParentFoldersIfMissing);
+			}
+
+			AssetRegistryModule.Get().AddPath(NewPath);
+
+			// If successful, this will move the contents (which at most will consist of directory placeholders) into the new folder and then delete the old folder
+			AssetViewUtils::RenameFolder(NewPath, OldPath);
+
+			TreeItem->Status = ETemplateFolderStatus::Exists;
+		}
+		else if (Response == EAppReturnType::No)
+		{
+			TreeItem->Status = ETemplateFolderStatus::MissingCreate;
+		}
+		else if (Response == EAppReturnType::Cancel)
+		{
+			return;
+		}
+	}
+
+	// Remove the old path from the active production's template, and add the new path instead
+	UProductionSettings* ProductionSettings = GetMutableDefault<UProductionSettings>();
+	FCinematicProduction* ActiveProduction = ProductionSettings->GetActiveProduction();
+	check(ActiveProduction);
+
+	ProductionSettings->RemoveTemplateFolder(ActiveProduction->ProductionID, OldPath);
+
+	const bool bCreateIfMissing = (TreeItem->Status == ETemplateFolderStatus::MissingDoNotCreate) ? false : true;
+	ProductionSettings->AddTemplateFolder(ActiveProduction->ProductionID, NewPath, bCreateIfMissing);
+
+	// Update the name input tree item and the path of all of its children (recursively)
+	TreeItem->Path.SetPathFromString(NewPath, EContentBrowserPathType::Internal);
+	SetChildrenPathRecursive(TreeItem, NewPath);
+
+	// Now that the path has changed, check again if this item exists, and update the status accordingly
+	if (AssetRegistryModule.Get().PathExists(NewPath))
+	{
+		TreeItem->Status = ETemplateFolderStatus::Exists;
+	}
+
+	// Sort the items at the same level in the tree as the input item to maintain a good ordering with the new folder name
+	TreeItem->Parent->Children.Sort([](const TSharedPtr<FTemplateFolderTreeItem>& A, const TSharedPtr<FTemplateFolderTreeItem>& B) { return A->Path.GetInternalPathString() < B->Path.GetInternalPathString(); });
+
+	TreeView->RequestTreeRefresh();
+}
+
+void SFolderHierarchyPanel::DeleteTemplateFolder(TSharedPtr<FTemplateFolderTreeItem> TreeItem)
+{
+	// Before removing any folders from the template, the user is prompted with a pop-up dialog confirming the action.
+
+	// If the folder being deleted exists in the content browser, and it is empty, the user is also prompted asking if they want to delete the Content Browser folder
+	if (TreeItem->Status == ETemplateFolderStatus::Exists && IsFolderEmpty(TreeItem))
+	{
+		const FText DialogMessage = LOCTEXT("DeleteExistingFolderDialogMessage", "An empty folder matching this template path exists in the Content Browser. Do you want to delete that empty folder?");
+		EAppReturnType::Type Response = FMessageDialog::Open(EAppMsgType::YesNoCancel, DialogMessage);
+
+		if (Response == EAppReturnType::Yes)
+		{
+			UDirectoryPlaceholderLibrary::DeletePlaceholdersInPath(TreeItem->Path.GetInternalPathString());
+			AssetViewUtils::DeleteFolders({ TreeItem->Path.GetInternalPathString() });
+		}
+		else if (Response == EAppReturnType::Cancel)
+		{
+			return;
+		}
+	}
+	else
+	{
+		const FText DialogMessage = LOCTEXT("DeleteTemplateDialogMessage", "Are you sure you want to delete this folder from the template?");
+		EAppReturnType::Type Response = FMessageDialog::Open(EAppMsgType::YesNo, DialogMessage);
+
+		if ((Response == EAppReturnType::No) || (Response == EAppReturnType::Cancel))
+		{
+			return;
+		}
+	}
+
+	// Recursively remove this item and all of its children from the active production's template and from the tree view
+	RemoveFolderFromTemplateRecursive(TreeItem);
+
+	// Now remove this item from its parent's list of children to finish cleaning up
+	TreeItem->Parent->Children.Remove(TreeItem);
+
+	TreeView->RequestTreeRefresh();
+}
+
+TSharedPtr<FTemplateFolderTreeItem> SFolderHierarchyPanel::FindItemAtPathRecursive(TSharedPtr<FTemplateFolderTreeItem> TreeItem, const FString& Path)
+{
+	if (TreeItem->Path.GetInternalPathString().Equals(Path))
+	{
+		return TreeItem;
+	}
+
+	for (const TSharedPtr<FTemplateFolderTreeItem>& Child : TreeItem->Children)
+	{
+		if (const TSharedPtr<FTemplateFolderTreeItem>& ItemAtPath = FindItemAtPathRecursive(Child, Path))
+		{
+			return ItemAtPath;
+		}
+	}
+
+	return nullptr;
+}
+
+void SFolderHierarchyPanel::CreateFolderFromTemplateRecursive(TSharedPtr<FTemplateFolderTreeItem> TreeItem)
+{
+	if (TreeItem->Status == ETemplateFolderStatus::MissingCreate)
+	{
+		const FString PathToCreate = TreeItem->Path.GetInternalPathString();
+		const FString RelativeFilePath = FPackageName::LongPackageNameToFilename(PathToCreate);
+		const FString AbsoluteFilePath = FPaths::ConvertRelativePathToFull(RelativeFilePath);
+
+		// Create the directory on disk, then add its path to the asset registry so it appears in Content Browser
+		if (!IFileManager::Get().DirectoryExists(*AbsoluteFilePath))
+		{
+			constexpr bool bCreateParentFoldersIfMissing = false;
+			if (IFileManager::Get().MakeDirectory(*AbsoluteFilePath, bCreateParentFoldersIfMissing))
+			{
+				FAssetRegistryModule& AssetRegistryModule = FModuleManager::GetModuleChecked<FAssetRegistryModule>("AssetRegistry");
+				AssetRegistryModule.Get().AddPath(PathToCreate);
+			}
+		}
+
+		TreeItem->Status = ETemplateFolderStatus::Exists;
+	}
+
+	// If we are not going to create this folder, we do not need to check any of its children
+	if (TreeItem->Status != ETemplateFolderStatus::MissingDoNotCreate)
+	{
+		for (TSharedPtr<FTemplateFolderTreeItem>& Child : TreeItem->Children)
+		{
+			CreateFolderFromTemplateRecursive(Child);
+		}
+	}
+}
+
+void SFolderHierarchyPanel::RemoveFolderFromTemplateRecursive(TSharedPtr<FTemplateFolderTreeItem> TreeItem)
+{
+	for (TSharedPtr<FTemplateFolderTreeItem>& Child : TreeItem->Children)
+	{
+		RemoveFolderFromTemplateRecursive(Child);
+	}
+
+	UProductionSettings* ProductionSettings = GetMutableDefault<UProductionSettings>();
+	const FCinematicProduction* ActiveProduction = ProductionSettings->GetActiveProduction();
+	check(ActiveProduction);
+
+	ProductionSettings->RemoveTemplateFolder(ActiveProduction->ProductionID, TreeItem->Path.GetInternalPathString());
+
+	TreeItem->Children.Empty();
+}
+
+void SFolderHierarchyPanel::SetChildrenPathRecursive(TSharedPtr<FTemplateFolderTreeItem> TreeItem, const FString& NewPath)
+{
+	FAssetRegistryModule& AssetRegistryModule = FModuleManager::GetModuleChecked<FAssetRegistryModule>("AssetRegistry");
+
+	UProductionSettings* ProductionSettings = GetMutableDefault<UProductionSettings>();
+	const FCinematicProduction* ActiveProduction = ProductionSettings->GetActiveProduction();
+	check(ActiveProduction);
+
+	for (TSharedPtr<FTemplateFolderTreeItem>& Child : TreeItem->Children)
+	{
+		const FString OldChildPath = Child->Path.GetInternalPathString();
+		const FString OldChildFolderName = FPaths::GetPathLeaf(OldChildPath);
+
+		// Update the child path with the new path of the parent
+		const FString NewChildPath = NewPath / OldChildFolderName;
+		Child->Path.SetPathFromString(NewChildPath, EContentBrowserPathType::Internal);
+
+		// Now that the path has changed, check again if this item exists, and update the status accordingly
+		if (AssetRegistryModule.Get().PathExists(NewChildPath))
+		{
+			Child->Status = ETemplateFolderStatus::Exists;
+		}
+		else if (TreeItem->Status == ETemplateFolderStatus::MissingDoNotCreate)
+		{
+			Child->Status = ETemplateFolderStatus::MissingDoNotCreate;
+		}
+		else
+		{
+			Child->Status = ETemplateFolderStatus::MissingCreate;
+		}
+
+		// Remove the old path from the active production's template, and add the new path instead
+		ProductionSettings->RemoveTemplateFolder(ActiveProduction->ProductionID, OldChildPath);
+
+		const bool bCreateIfMissing = (Child->Status == ETemplateFolderStatus::MissingDoNotCreate) ? false : true;
+		ProductionSettings->AddTemplateFolder(ActiveProduction->ProductionID, NewChildPath, bCreateIfMissing);
+
+		SetChildrenPathRecursive(Child, NewChildPath);
+	}
+}
+
+void SFolderHierarchyPanel::SetTemplateFolderStatusRecursive(TSharedPtr<FTemplateFolderTreeItem> TreeItem, ETemplateFolderStatus NewStatus)
+{
+	TreeItem->Status = NewStatus;
+
+	UProductionSettings* ProductionSettings = GetMutableDefault<UProductionSettings>();
+	const FCinematicProduction* ActiveProduction = ProductionSettings->GetActiveProduction();
+	check(ActiveProduction);
+
+	const FString Path = TreeItem->Path.GetInternalPathString();
+	ProductionSettings->RemoveTemplateFolder(ActiveProduction->ProductionID, Path);
+
+	const bool bCreateIfMissing = (TreeItem->Status == ETemplateFolderStatus::MissingDoNotCreate) ? false : true;
+	ProductionSettings->AddTemplateFolder(ActiveProduction->ProductionID, Path, bCreateIfMissing);
+
+	for (TSharedPtr<FTemplateFolderTreeItem>& Child : TreeItem->Children)
+	{
+		SetTemplateFolderStatusRecursive(Child, NewStatus);
+	}
+}
+
+FString SFolderHierarchyPanel::CreateUniqueFolderName(FContentBrowserItemPath InPath)
+{
+	// This method for creating a unique folder name is adapted from the implementation used by the Content Browser to achieve a similar convention
+	const FText DefaultFolderBaseName = LOCTEXT("DefaultFolderName", "NewFolder");
+	UContentBrowserDataSubsystem* ContentBrowserDataSubsystem = IContentBrowserDataModule::Get().GetSubsystem();
+
+	const FString DefaultFolderName = DefaultFolderBaseName.ToString();
+	int32 NewFolderPostfix = 0;
+	FName CombinedPathName;
+	for (;;)
+	{
+		FString CombinedPathNameStr = InPath.GetVirtualPathString() / DefaultFolderName;
+		if (NewFolderPostfix > 0)
+		{
+			CombinedPathNameStr.AppendInt(NewFolderPostfix);
+		}
+		++NewFolderPostfix;
+
+		CombinedPathName = *CombinedPathNameStr;
+
+		// Check if the path matches a folder that already exists in the content browser
+		const FContentBrowserItem ExistingFolder = ContentBrowserDataSubsystem->GetItemAtPath(CombinedPathName, EContentBrowserItemTypeFilter::IncludeFolders);
+
+		// The path used by the Content Browser subsystem function was a virtual path, but our function operates on internal paths, so we have to strip "/All" off the front
+		CombinedPathNameStr.RemoveFromStart(TEXT("/All"));
+
+		// Check if the path matches a folder in our template
+		TSharedPtr<FTemplateFolderTreeItem> Item = FindItemAtPathRecursive(RootItem, CombinedPathNameStr);
+
+		if (!ExistingFolder.IsValid() && !Item.IsValid())
+		{
+			break;
+		}
+	}
+
+	return FPaths::GetPathLeaf(CombinedPathName.ToString());
+}
+
+bool SFolderHierarchyPanel::IsFolderEmpty(TSharedPtr<FTemplateFolderTreeItem> TreeItem) const
+{
+	FAssetRegistryModule& AssetRegistryModule = FModuleManager::Get().GetModuleChecked<FAssetRegistryModule>(TEXT("AssetRegistry"));
+
+	FARFilter Filter;
+	Filter.bRecursivePaths = true;
+	Filter.PackagePaths.Add(TreeItem->Path.GetInternalPathName());
+
+	// Find all of the assets (recursively) in the input folder
+	TArray<FAssetData> AssetDataList;
+	AssetRegistryModule.Get().GetAssets(Filter, AssetDataList);
+
+	// Check if all of the assets found are placeholders
+	bool bIsEmpty = true;
+	for (const FAssetData& AssetData : AssetDataList)
+	{
+		if (AssetData.GetClass() != UDirectoryPlaceholder::StaticClass())
+		{
+			bIsEmpty = false;
+		}
+	}
+
+	return bIsEmpty;
+}
+
+const FSlateBrush* SFolderHierarchyPanel::GetFolderIconBadge(TSharedPtr<FTemplateFolderTreeItem> TreeItem) const
+{
+	if (TreeItem->Status == ETemplateFolderStatus::MissingCreate)
+	{
+		const FSlateBrush* WarningBrush = FSlateIcon(FShotsEditorStyle::StyleName, "Badges.FolderNew").GetIcon();
+		return WarningBrush;
+	}
+	else if (TreeItem->Status == ETemplateFolderStatus::MissingDoNotCreate)
+	{
+		const FSlateBrush* WarningBrush = FSlateIcon(FShotsEditorStyle::StyleName, "Badges.FolderRename").GetIcon();
+		return WarningBrush;
+	}
+	return nullptr;
+}
+
+FText SFolderHierarchyPanel::GetFolderStatusText(TSharedPtr<FTemplateFolderTreeItem> TreeItem) const
+{
+	switch (TreeItem->Status)
+	{
+	case ETemplateFolderStatus::MissingCreate:
+		return LOCTEXT("Status_MissingCreate", "Create");
+	case ETemplateFolderStatus::MissingDoNotCreate:
+		return LOCTEXT("Status_MissingDoNotCreate", "Do Not Create");
+	default:
+		return FText::GetEmpty();
+	}
+}
+
+#undef LOCTEXT_NAMESPACE
diff --git a/Engine/Plugins/VirtualProduction/Shots/Source/ShotsEditor/Private/UI/SFolderHierarchyPanel.h b/Engine/Plugins/VirtualProduction/Shots/Source/ShotsEditor/Private/UI/SFolderHierarchyPanel.h
new file mode 100644
--- /dev/null
+++ b/Engine/Plugins/VirtualProduction/Shots/Source/ShotsEditor/Private/UI/SFolderHierarchyPanel.h
@@ -0,0 +1,154 @@
+// Copyright Epic Games, Inc. All Rights Reserved.
+
+#pragma once
+
+#include "Widgets/SCompoundWidget.h"
+
+#include "ContentBrowserItemPath.h"
+#include "ProductionSettings.h"
+#include "Widgets/Text/SInlineEditableTextBlock.h"
+#include "Widgets/Views/STreeView.h"
+
+/** States in which a template folder can exist */
+enum class ETemplateFolderStatus : uint8
+{
+	/** The status has not yet been determined */
+	None,
+
+	/** Template folder path exists in the Project Content folder */
+	Exists,
+
+	/** Template folder path does not exist in the Project Content folder, but should be created OnApply */
+	MissingCreate,
+
+	/** Template folder path does not exist in the Project Content folder, but should not be created OnApply */
+	MissingDoNotCreate,
+};
+
+/** An entry in the template folder tree view with knowledge of its path, children, and parent */
+struct FTemplateFolderTreeItem
+{
+	/** The path of this template folder */
+	FContentBrowserItemPath Path;
+
+	/** The parent of this item in the tree */
+	TSharedPtr<FTemplateFolderTreeItem> Parent;
+
+	/** The children of this item in the tree */
+	TArray<TSharedPtr<FTemplateFolderTreeItem>> Children;
+
+	/** The status of this template folder, indicating whether it needs to be created OnApply */
+	ETemplateFolderStatus Status = ETemplateFolderStatus::None;
+
+	/** The text widget that display this item's name and supports renaming */
+	TSharedPtr<SInlineEditableTextBlock> NameWidget;
+};
+
+/**
+ * UI for the Folder Hierarchy panel in the Production Wizard
+ */
+class SFolderHierarchyPanel : public SCompoundWidget
+{
+public:
+	SLATE_BEGIN_ARGS(SFolderHierarchyPanel) {}
+	SLATE_END_ARGS()
+
+	SFolderHierarchyPanel() = default;
+	~SFolderHierarchyPanel();
+
+	void Construct(const FArguments& InArgs);
+
+private:
+	/** Generates the row widget for an entry in the tree view */
+	TSharedRef<ITableRow> OnGenerateRow(TSharedPtr<FTemplateFolderTreeItem> TreeItem, const TSharedRef<STableViewBase>& OwnerTable);
+
+	/** Gets the children of the input tree view item to build additional tree rows */
+	void OnGetChildren(TSharedPtr<FTemplateFolderTreeItem> TreeItem, TArray<TSharedPtr<FTemplateFolderTreeItem>>& OutNodes);
+
+	/** Callback when the tree view rebuilds itself, used here to allow the user to immediately rename the item after the tree view updates to show it */
+	void OnItemsRebuilt();
+
+	/** Set the expansion state of the input treem item and all of its children (recursively) */
+	void SetExpansionRecursive(TSharedPtr<FTemplateFolderTreeItem> TreeItem, bool bInExpand) const;
+
+	/** Spawns a context menu when the user right-clicks on an entry in the tree view */
+	TSharedPtr<SWidget> OnContextMenuOpening();
+
+	/** Handles key presses on the tree view */
+	FReply OnKeyDown(const FGeometry& MyGeometry, const FKeyEvent& InKeyEvent);
+
+	/** Builds the template folder list items (displayed in the tree view) using the current active production's folder template */
+	void BuildTreeFromProductionTemplate();
+
+	/** Adds a new folder to the template under the currently selected tree item */
+	FReply OnAddFolderToTemplate();
+
+	/** Creates any template folders whose status is MissingCreate in the Content Browser */
+	FReply OnCreateTemplateFolders();
+
+	/** Resets the active production's template to the previously cached state and updates the tree view */
+	FReply OnResetTemplateChanges();
+
+	/** Returns true if the input text is a valid folder name, otherwise false and fills the output error message accordingly */
+	bool IsValidFolderName(const FText& InText, FText& OutErrorMessage, TSharedPtr<FTemplateFolderTreeItem> TreeItem);
+
+	/** 
+	 * Sets the name of the input template folder.
+	 * If the folder exists in the Content Browser, and is empty, the user is prompted with the option to also rename the Content Browser folder
+	 */
+	void SetTemplateFolderName(const FText& InText, ETextCommit::Type InCommitType, TSharedPtr<FTemplateFolderTreeItem> TreeItem);
+
+	/** 
+	 * Deletes the input template folder from the template.
+	 * If the folder exists in the Content Browser, and is empty, the user is prompted with the option to also delete the Content Browser folder
+	 */
+	void DeleteTemplateFolder(TSharedPtr<FTemplateFolderTreeItem> SelectedTreeItem);
+
+	/** Returns the template folder tree item whose path matches the input path */
+	TSharedPtr<FTemplateFolderTreeItem> FindItemAtPathRecursive(TSharedPtr<FTemplateFolderTreeItem> TreeItem, const FString& Path);
+
+	/** Walks the template folder tree and creates all folders marked as MissingCreate */
+	void CreateFolderFromTemplateRecursive(TSharedPtr<FTemplateFolderTreeItem> TreeItem);
+
+	/** Removes the input tree item and all of its children (recursively) from the template */
+	void RemoveFolderFromTemplateRecursive(TSharedPtr<FTemplateFolderTreeItem> TreeItem);
+
+	/** Modifies the path of the input tree item and all of its children (recursively) with the input new path */
+	void SetChildrenPathRecursive(TSharedPtr<FTemplateFolderTreeItem> TreeItem, const FString& NewPath);
+
+	/** Sets the folder status of the input tree item and all of its children (recursively) */
+	void SetTemplateFolderStatusRecursive(TSharedPtr<FTemplateFolderTreeItem> TreeItem, ETemplateFolderStatus NewStatus);
+
+	/** Returns a valid folder name that is unique at the input path in both the Content Browser and the Folder Template */
+	FString CreateUniqueFolderName(FContentBrowserItemPath InPath);
+
+	/** Returns true if the input tree item exists in the Content Browser and contains any assets (other than placeholders) */
+	bool IsFolderEmpty(TSharedPtr<FTemplateFolderTreeItem> TreeItem) const;
+
+	/** Returns the badge icon to overlay on top of the folder icon for the input tree item */
+	const FSlateBrush* GetFolderIconBadge(TSharedPtr<FTemplateFolderTreeItem> TreeItem) const;
+
+	/** Returns the status text to display next to the input tree item in its row in the tree view */
+	FText GetFolderStatusText(TSharedPtr<FTemplateFolderTreeItem> TreeItem) const;
+
+private:
+	/* Root path of the template folder tree */
+	static FName RootPath;
+
+	/** The source template folder items for the tree view */
+	TArray<TSharedPtr<FTemplateFolderTreeItem>> FolderItemsSource;
+
+	/** The tree view of template folders for the current active production */
+	TSharedPtr<STreeView<TSharedPtr<FTemplateFolderTreeItem>>> TreeView;
+
+	/** Convenient alias for the root item (first item) in the tree view */
+	TSharedPtr<FTemplateFolderTreeItem> RootItem;
+
+	/** The cached state of the template, used to reset user changes */
+	TArray<FFolderTemplate> CachedInitialState;
+
+	/** The mostly recently added tree item, used to allow the user to immediately rename the item after the tree view updates to show it */
+	TSharedPtr<FTemplateFolderTreeItem> MostRecentlyAddedItem;
+
+	FDelegateHandle ActiveProductionChangedHandle;
+};
diff --git a/Engine/Plugins/VirtualProduction/Shots/Source/ShotsEditor/Private/UI/SProductionWizard.cpp b/Engine/Plugins/VirtualProduction/Shots/Source/ShotsEditor/Private/UI/SProductionWizard.cpp
--- a/Engine/Plugins/VirtualProduction/Shots/Source/ShotsEditor/Private/UI/SProductionWizard.cpp
+++ b/Engine/Plugins/VirtualProduction/Shots/Source/ShotsEditor/Private/UI/SProductionWizard.cpp
@@ -6,6 +6,7 @@
 #include "ShotsEditorStyle.h"
 #include "UI/SNamingTokensPanel.h"
 #include "UI/SAssetNamingPanel.h"
+#include "UI/SFolderHierarchyPanel.h"
 #include "UI/SProductionListPanel.h"
 #include "UI/SRevisionControlPanel.h"
 #include "UI/SSequencerSettingsPanel.h"
@@ -223,7 +224,7 @@ TSharedRef<SWidget> SProductionWizard::MakeMenuPanel()
 	TSharedPtr<FProductionWizardMenuEntry> FolderHierarchyMenuEntry = MakeShared<FProductionWizardMenuEntry>();
 	FolderHierarchyMenuEntry->Label = LOCTEXT("FolderHierarchyMenuLabel", "Folder Hierarchy");
 	FolderHierarchyMenuEntry->Icon = FSlateIcon(FShotsEditorStyle::StyleName, "Icons.Folder");
-	FolderHierarchyMenuEntry->Panel = MakeFolderHierarchyPanel();
+	FolderHierarchyMenuEntry->Panel = SNew(SFolderHierarchyPanel);
 
 	ProductionMenuEntries.Add(SequencerSettingsMenuEntry);
 	ProductionMenuEntries.Add(NamingTokensMenuEntry);
@@ -329,24 +330,4 @@ void SProductionWizard::OnProductionMenuSelectionChanged(TSharedPtr<FProductionW
 	}
 }
 
-TSharedRef<SWidget> SProductionWizard::MakeFolderHierarchyPanel()
-{
-	// TODO: Fill in this stubbed out panel
-	return SNew(SVerticalBox)
-
-	+ SVerticalBox::Slot()
-		.FillHeight(1.0f)
-		.HAlign(HAlign_Center)
-		[
-			SNew(SHorizontalBox)
-
-			+ SHorizontalBox::Slot()
-				.FillWidth(1.0f)
-				.VAlign(VAlign_Center)
-				[
-					SNew(STextBlock).Text(LOCTEXT("FolderHierarchyPanel", "FOLDER HIERARCHY PANEL"))
-				]
-		];
-}
-
 #undef LOCTEXT_NAMESPACE
diff --git a/Engine/Plugins/VirtualProduction/Shots/Source/ShotsEditor/Private/UI/SProductionWizard.h b/Engine/Plugins/VirtualProduction/Shots/Source/ShotsEditor/Private/UI/SProductionWizard.h
--- a/Engine/Plugins/VirtualProduction/Shots/Source/ShotsEditor/Private/UI/SProductionWizard.h
+++ b/Engine/Plugins/VirtualProduction/Shots/Source/ShotsEditor/Private/UI/SProductionWizard.h
@@ -42,9 +42,6 @@ private:
 	/** Creates the buttons on the bottom of the production wizard */
 	TSharedRef<SWidget> MakeButtonsPanel();
 
-	/** Functions for creating the panels for each of the production wizard's menu entries */
-	TSharedRef<SWidget> MakeFolderHierarchyPanel();
-
 	/** Generates a row in the menu list views */
 	TSharedRef<ITableRow> OnGenerateMenuRow(TSharedPtr<FProductionWizardMenuEntry> MenuEntry, const TSharedRef<STableViewBase>& OwnerTable);
 
diff --git a/Engine/Plugins/VirtualProduction/Shots/Source/ShotsEditor/ShotsEditor.build.cs b/Engine/Plugins/VirtualProduction/Shots/Source/ShotsEditor/ShotsEditor.build.cs
--- a/Engine/Plugins/VirtualProduction/Shots/Source/ShotsEditor/ShotsEditor.build.cs
+++ b/Engine/Plugins/VirtualProduction/Shots/Source/ShotsEditor/ShotsEditor.build.cs
@@ -16,12 +16,15 @@ namespace UnrealBuildTool.Rules
 				new string[]
 				{
 					"AssetDefinition",
+					"AssetRegistry",
 					"AssetTools",
 					"ClassViewer",
+					"ContentBrowserData",
 					"Core",
 					"CoreUObject",
 					"DesktopPlatform",
 					"DeveloperSettings",
+					"DirectoryPlaceholder",
 					"Engine",
 					"InputCore",
 					"Json",
