ShotManagement: Fix editor crash when a new template folder is added to the tree outside of the current scroll area.

#rnx

diff --git a/Engine/Plugins/VirtualProduction/Shots/Source/ShotsEditor/Private/UI/SFolderHierarchyPanel.cpp b/Engine/Plugins/VirtualProduction/Shots/Source/ShotsEditor/Private/UI/SFolderHierarchyPanel.cpp
--- a/Engine/Plugins/VirtualProduction/Shots/Source/ShotsEditor/Private/UI/SFolderHierarchyPanel.cpp
+++ b/Engine/Plugins/VirtualProduction/Shots/Source/ShotsEditor/Private/UI/SFolderHierarchyPanel.cpp
@@ -311,7 +311,7 @@ void SFolderHierarchyPanel::OnGetChildren(TSharedPtr<FTemplateFolderTreeItem> Tr
 void SFolderHierarchyPanel::OnItemsRebuilt()
 {
 	// Upon regenerating the tree view, allow the user to immediately interact with the name widget of the newly added template folder in order to rename it
-	if (MostRecentlyAddedItem)
+	if (MostRecentlyAddedItem && MostRecentlyAddedItem->NameWidget)
 	{
 		FSlateApplication::Get().SetKeyboardFocus(MostRecentlyAddedItem->NameWidget.ToSharedRef());
 		MostRecentlyAddedItem->NameWidget->EnterEditingMode();
