ShotManagement: Fix bug where productions can sometimes be created with no name.
#rnx

diff --git a/Engine/Plugins/VirtualProduction/Shots/Source/ShotsEditor/Private/ProductionSettings.cpp b/Engine/Plugins/VirtualProduction/Shots/Source/ShotsEditor/Private/ProductionSettings.cpp
--- a/Engine/Plugins/VirtualProduction/Shots/Source/ShotsEditor/Private/ProductionSettings.cpp
+++ b/Engine/Plugins/VirtualProduction/Shots/Source/ShotsEditor/Private/ProductionSettings.cpp
@@ -282,14 +282,9 @@ void UProductionSettings::PostEditChangeProperty(struct FPropertyChangedEvent& P
 
 void UProductionSettings::AddProduction()
 {
-	const int32 ArrayIndex = Productions.Add(FCinematicProduction());
-
-	// Reset the active production pointer, in case the array was reallocated
-	ActiveProduction = Productions.FindByPredicate([this](const FCinematicProduction& Other) { return Other.ProductionName == ActiveProductionName; });
-
-	TryUpdateDefaultConfigFile();
-
-	ProductionListChangedDelegate.Broadcast();
+	FCinematicProduction NewProduction;
+	NewProduction.ProductionName = GetUniqueProductionName();
+	AddProduction(NewProduction);
 }
 
 void UProductionSettings::AddProduction(const FCinematicProduction& ProductionToAdd)
@@ -342,6 +337,9 @@ void UProductionSettings::DeleteProduction(FGuid ProductionID)
 
 		Productions.RemoveAt(IndexToDelete);
 
+		// Reset the active production pointer
+		ActiveProduction = Productions.FindByPredicate([this](const FCinematicProduction& Other) { return Other.ProductionName == ActiveProductionName; });
+
 		if (bIsActiveProduction)
 		{
 			SetActiveProduction(ProductionID);
diff --git a/Engine/Plugins/VirtualProduction/Shots/Source/ShotsEditor/Private/UI/SProductionListPanel.cpp b/Engine/Plugins/VirtualProduction/Shots/Source/ShotsEditor/Private/UI/SProductionListPanel.cpp
--- a/Engine/Plugins/VirtualProduction/Shots/Source/ShotsEditor/Private/UI/SProductionListPanel.cpp
+++ b/Engine/Plugins/VirtualProduction/Shots/Source/ShotsEditor/Private/UI/SProductionListPanel.cpp
@@ -28,10 +28,12 @@ void SProductionListPanel::Construct(const FArguments& InArgs)
 	// Build the list of productions for the list view before building the widget for the first time
 	UpdateProductionList();
 
-	auto OnCreateNewProductionClicked = []() -> FReply
+	auto OnCreateNewProductionClicked = [this]() -> FReply
 		{
 			UProductionSettings* ProductionSettings = GetMutableDefault<UProductionSettings>();
 			ProductionSettings->AddProduction();
+
+			MostRecentProductionID = ProductionSettings->Productions.Last().ProductionID;
 			return FReply::Handled();
 		};
 
@@ -178,11 +180,11 @@ void SProductionListPanel::UpdateProductionList()
 
 void SProductionListPanel::OnProductionListItemsRebuilt()
 {
-	// If there was a production with no name, then ProductionListButtonToRename will be valid, and it can be put into edit mode so the user can rename the production.
 	if (ProductionListButtonToRename)
 	{
 		ProductionListButtonToRename->EnterEditMode();
-		ProductionListButtonToRename = nullptr;
+		ProductionListButtonToRename.Reset();
+		MostRecentProductionID.Invalidate();
 	}
 }
 
@@ -198,10 +200,10 @@ TSharedRef<ITableRow> SProductionListPanel::OnGenerateProductionRow(TSharedPtr<F
 	TSharedRef<SProductionListButton> ListButton = SNew(SProductionListButton)
 		.ProductionID(InItem->ProductionID);
 
-	// If this production was newly created, it will not have a name yet, so the user should be prompted to give it one.
+	// If the production shown in this row was just created, it will have a valid default name, but the user should have an opportunity to immediately rename it.
 	// However, the production's list button cannot be put into edit mode and focused until after the widget is created.
 	// Therefore, the widget is saved, and will be put into edit mode after the items for this ListView are finished being rebuilt. 
-	if (InItem->ProductionName.IsEmpty())
+	if (InItem->ProductionID.IsValid() && InItem->ProductionID == MostRecentProductionID)
 	{
 		ProductionListButtonToRename = ListButton;
 	}
@@ -389,10 +391,6 @@ void SProductionListButton::Construct(const FArguments& InArgs)
 			const UProductionSettings* const ProductionSettings = GetDefault<UProductionSettings>();
 			if (const FCinematicProduction* Production = ProductionSettings->GetProduction(ProductionID.Get()))
 			{
-				if (Production->ProductionName.IsEmpty())
-				{
-					return FText::FromString(ProductionSettings->GetUniqueProductionName());
-				}
 				return FText::FromString(Production->ProductionName);
 			}
 			return FText::GetEmpty();
diff --git a/Engine/Plugins/VirtualProduction/Shots/Source/ShotsEditor/Private/UI/SProductionListPanel.h b/Engine/Plugins/VirtualProduction/Shots/Source/ShotsEditor/Private/UI/SProductionListPanel.h
--- a/Engine/Plugins/VirtualProduction/Shots/Source/ShotsEditor/Private/UI/SProductionListPanel.h
+++ b/Engine/Plugins/VirtualProduction/Shots/Source/ShotsEditor/Private/UI/SProductionListPanel.h
@@ -89,6 +89,9 @@ private:
 	/** Button belonging to one of the production list view rows that needs to be put into edit mode for the user to rename its production */
 	TSharedPtr<SProductionListButton> ProductionListButtonToRename;
 
+	/** Valid immediately after a new production is added, used to determine which production list item should be put into edit mode */
+	FGuid MostRecentProductionID;
+
 	/** Handle to the delegate which responds to changes in UProductionSettings list of productions */
 	FDelegateHandle ProductionListChangedHandle;
 };
