ShotManagement: Fix for two issues with importing productions:

Productions with names ending in numbers were being incorrectly renamed on import.
Users are now warned if they try to import a production whose name matches an existing production.
Users now get a pop-up error message if the selected JSON file cannot be parsed as a valid production.

#rnx

diff --git a/Engine/Plugins/VirtualProduction/Shots/Source/ShotsEditor/Private/UI/SProductionListPanel.cpp b/Engine/Plugins/VirtualProduction/Shots/Source/ShotsEditor/Private/UI/SProductionListPanel.cpp
--- a/Engine/Plugins/VirtualProduction/Shots/Source/ShotsEditor/Private/UI/SProductionListPanel.cpp
+++ b/Engine/Plugins/VirtualProduction/Shots/Source/ShotsEditor/Private/UI/SProductionListPanel.cpp
@@ -9,6 +9,7 @@
 #include "HAL/FileManager.h"
 #include "JsonObjectConverter.h"
 #include "Layout/WidgetPath.h"
+#include "Misc/MessageDialog.h"
 #include "ProductionSettings.h"
 #include "ShotsEditorStyle.h"
 #include "Styling/StyleColors.h"
@@ -280,17 +281,44 @@ FReply SProductionListPanel::ImportProduction()
 			constexpr int64 CheckFlags = 0;
 			constexpr int64 SkipFlags = 0;
 			constexpr bool bStrictMode = true;
-			if (FJsonObjectConverter::JsonObjectToUStruct<FCinematicProduction>(JsonObject.ToSharedRef(), ImportedProduction.Get(), CheckFlags, SkipFlags, bStrictMode))
+			FText ErrorReason;
+			const bool bSuccessfulImport = FJsonObjectConverter::JsonObjectToUStruct<FCinematicProduction>(JsonObject.ToSharedRef(), ImportedProduction.Get(), CheckFlags, SkipFlags, bStrictMode, &ErrorReason);
+
+			if (bSuccessfulImport)
 			{
-				// If the import and conversion from json was successful, add the imported production to the Production Project Settings list of productions
 				UProductionSettings* ProductionSettings = GetMutableDefault<UProductionSettings>();
-
-				// Give the imported production a new unique ID and ensure that its name is unique. 
 				ImportedProduction->ProductionID = FGuid::NewGuid();
-				ImportedProduction->ProductionName = ProductionSettings->GetUniqueProductionName(ImportedProduction->ProductionName);
 
-				ProductionSettings->AddProduction(*ImportedProduction);
+				// If the name of the production being imported already exists in the list of productions, prompt the user to confirm before actually adding it
+				const bool bProductionNameAlreadyExists = ProductionSettings->Productions.ContainsByPredicate([ProductionName = ImportedProduction->ProductionName](const FCinematicProduction& Other) { return Other.ProductionName == ProductionName; });
+
+				bool bShouldAddProduction = true;
+				if (bProductionNameAlreadyExists)
+				{
+					const FText DialogMessage = FText::Format(LOCTEXT("ProductionNameAlreadyExistsMessage", "A Production named {0} already exists. Do you still want to import this production?"), FText::FromString(ImportedProduction->ProductionName));
+					const EAppReturnType::Type Response = FMessageDialog::Open(EAppMsgType::YesNo, DialogMessage);
+
+					if (Response != EAppReturnType::Yes)
+					{
+						bShouldAddProduction = false;
+					}
+				}
+
+				if (bShouldAddProduction)
+				{
+					ProductionSettings->AddProduction(*ImportedProduction);
+				}
 			}
+			else
+			{
+				const FText ErrorMessage = FText::Format(LOCTEXT("SpecificImportErrorMessage", "The selected file failed to import.\n\n{0}"), ErrorReason);
+				FMessageDialog::Open(EAppMsgType::Ok, ErrorMessage);
+			}
+		}
+		else
+		{
+			const FText ErrorMessage = LOCTEXT("GenericImportErrorMessage", "The selected file failed to import.\n\nThe JSON file may be incorrectly formatted.");
+			FMessageDialog::Open(EAppMsgType::Ok, ErrorMessage);
 		}
 	}
 
