ShotManagement: Add new Cine Asset type.

UCineAsset is a subclass of ULevelSequence, and behaves very much like a level sequence, but it also has a level associated with it that will load by default when the user opens the asset. Future work will expand on its functionality to add additional metadata and tools for cinematic assets.

Before creating a new cine asset, the asset factory will spawn a window for the user to configure the properties of the new asset that is about to be created. This is where the user is first prompted to set the associated to a level.

diff --git a/Engine/Plugins/MovieScene/LevelSequenceEditor/Source/LevelSequenceEditor/Private/AssetTools/AssetDefinition_LevelSequence.cpp b/Engine/Plugins/MovieScene/LevelSequenceEditor/Source/LevelSequenceEditor/Private/AssetTools/AssetDefinition_LevelSequence.cpp
--- a/Engine/Plugins/MovieScene/LevelSequenceEditor/Source/LevelSequenceEditor/Private/AssetTools/AssetDefinition_LevelSequence.cpp
+++ b/Engine/Plugins/MovieScene/LevelSequenceEditor/Source/LevelSequenceEditor/Private/AssetTools/AssetDefinition_LevelSequence.cpp
@@ -1,6 +1,6 @@
 // Copyright Epic Games, Inc. All Rights Reserved.
 
-#include "AssetDefinition_LevelSequence.h"
+#include "AssetTools/AssetDefinition_LevelSequence.h"
 #include "LevelSequenceEditorToolkit.h"
 #include "LevelSequenceEditorStyle.h"
 
diff --git a/Engine/Plugins/MovieScene/LevelSequenceEditor/Source/LevelSequenceEditor/Private/AssetTools/AssetDefinition_LevelSequence.h b/Engine/Plugins/MovieScene/LevelSequenceEditor/Source/LevelSequenceEditor/Private/AssetTools/AssetDefinition_LevelSequence.h
deleted file mode 100644
diff --git a/Engine/Plugins/MovieScene/LevelSequenceEditor/Source/LevelSequenceEditor/Public/AssetTools/AssetDefinition_LevelSequence.h b/Engine/Plugins/MovieScene/LevelSequenceEditor/Source/LevelSequenceEditor/Public/AssetTools/AssetDefinition_LevelSequence.h
new file mode 100644
--- /dev/null
+++ b/Engine/Plugins/MovieScene/LevelSequenceEditor/Source/LevelSequenceEditor/Public/AssetTools/AssetDefinition_LevelSequence.h
@@ -0,0 +1,28 @@
+// Copyright Epic Games, Inc. All Rights Reserved.
+
+#pragma once
+
+#include "LevelSequence.h"
+#include "AssetDefinitionDefault.h"
+
+#include "AssetDefinition_LevelSequence.generated.h"
+
+UCLASS()
+class LEVELSEQUENCEEDITOR_API UAssetDefinition_LevelSequence : public UAssetDefinitionDefault
+{
+	GENERATED_BODY()
+
+public:
+	// UAssetDefinition Begin
+	virtual FText GetAssetDisplayName() const override { return NSLOCTEXT("AssetTypeActions", "AssetTypeActions_LevelSequence", "Level Sequence"); }
+	virtual FLinearColor GetAssetColor() const override { return FLinearColor(FColor(200, 80, 80)); }
+	virtual TSoftClassPtr<UObject> GetAssetClass() const override { return ULevelSequence::StaticClass(); }
+	virtual TConstArrayView<FAssetCategoryPath> GetAssetCategories() const override
+	{
+		static const auto Categories = { EAssetCategoryPaths::Basic, EAssetCategoryPaths::Cinematics };
+		return Categories;
+	}
+	virtual FAssetOpenSupport GetAssetOpenSupport(const FAssetOpenSupportArgs& OpenSupportArgs) const override;
+	virtual EAssetCommandResult OpenAssets(const FAssetOpenArgs& OpenArgs) const override;
+	// UAssetDefinition End
+};
diff --git a/Engine/Plugins/VirtualProduction/Shots/Source/ShotsEditor/Private/AssetDefinition_CineAsset.cpp b/Engine/Plugins/VirtualProduction/Shots/Source/ShotsEditor/Private/AssetDefinition_CineAsset.cpp
new file mode 100644
--- /dev/null
+++ b/Engine/Plugins/VirtualProduction/Shots/Source/ShotsEditor/Private/AssetDefinition_CineAsset.cpp
@@ -0,0 +1,90 @@
+// Copyright Epic Games, Inc. All Rights Reserved.
+
+#include "AssetDefinition_CineAsset.h"
+#include "CineAsset.h"
+
+#include "Editor.h"
+#include "FileHelpers.h"
+
+#define LOCTEXT_NAMESPACE "AssetTypeActions"
+
+DEFINE_LOG_CATEGORY_STATIC(LogCineAssetDefinition, Log, All)
+
+FText UAssetDefinition_CineAsset::GetAssetDisplayName() const
+{
+	return LOCTEXT("AssetTypeActions_CineAsset", "Cine Asset");
+}
+
+TSoftClassPtr<> UAssetDefinition_CineAsset::GetAssetClass() const
+{
+	return UCineAsset::StaticClass();
+}
+
+FLinearColor UAssetDefinition_CineAsset::GetAssetColor() const
+{
+	return FColor(172, 35, 85);
+}
+
+TConstArrayView<FAssetCategoryPath> UAssetDefinition_CineAsset::GetAssetCategories() const
+{
+	static const FAssetCategoryPath Categories[] = { EAssetCategoryPaths::Cinematics };
+	return Categories;
+}
+
+EAssetCommandResult UAssetDefinition_CineAsset::OpenAssets(const FAssetOpenArgs& OpenArgs) const
+{
+	return Super::OpenAssets(OpenArgs);
+}
+
+TArray<FAssetData> UAssetDefinition_CineAsset::PrepareToActivateAssets(const FAssetActivateArgs& ActivateArgs) const
+{
+	TArray<FAssetData> AssetsToOpen;
+
+	// We only support opening one cine asset at a time
+	if (ActivateArgs.Assets.Num() == 1)
+	{
+		const FAssetData& CineAssetData = ActivateArgs.Assets[0];
+		AssetsToOpen.Add(CineAssetData);
+
+		if (UCineAsset* CineAsset = Cast<UCineAsset>(CineAssetData.GetAsset()))
+		{
+			if (CineAsset->Level.IsValid())
+			{
+				if (UWorld* WorldToOpen = Cast<UWorld>(CineAsset->Level.TryLoad()))
+				{
+					UWorld* CurrentWorld = GEditor->GetEditorWorldContext().World();
+					if (CurrentWorld != WorldToOpen)
+					{
+						// Prompt the user to save their unsaved changes to the current level before loading the level associated with this asset
+						constexpr bool bPromptUserToSave = true;
+						constexpr bool bSaveMapPackages = true;
+						constexpr bool bSaveContentPackages = false;
+						if (FEditorFileUtils::SaveDirtyPackages(bPromptUserToSave, bSaveMapPackages, bSaveContentPackages))
+						{
+							if (!WorldToOpen->GetPackage()->HasAnyPackageFlags(PKG_NewlyCreated))
+							{
+								const FString FileToOpen = FPackageName::LongPackageNameToFilename(WorldToOpen->GetOutermost()->GetName(), FPackageName::GetMapPackageExtension());
+								const bool bLoadAsTemplate = false;
+								const bool bShowProgress = true;
+								FEditorFileUtils::LoadMap(FileToOpen, bLoadAsTemplate, bShowProgress);
+							}
+						}
+						else
+						{
+							// If the user canceled out of the prompt to save the current level, then do not try to open the cine asset
+							AssetsToOpen.Empty();
+						}
+					}
+				}
+				else
+				{
+					UE_LOG(LogCineAssetDefinition, Error, TEXT("Failed to load %s while opening %s"), *CineAsset->Level.ToString(), *CineAsset->GetFName().ToString());
+				}
+			}
+		}
+	}
+
+	return AssetsToOpen;
+}
+
+#undef LOCTEXT_NAMESPACE
diff --git a/Engine/Plugins/VirtualProduction/Shots/Source/ShotsEditor/Private/AssetDefinition_CineAsset.h b/Engine/Plugins/VirtualProduction/Shots/Source/ShotsEditor/Private/AssetDefinition_CineAsset.h
new file mode 100644
--- /dev/null
+++ b/Engine/Plugins/VirtualProduction/Shots/Source/ShotsEditor/Private/AssetDefinition_CineAsset.h
@@ -0,0 +1,24 @@
+// Copyright Epic Games, Inc. All Rights Reserved.
+
+#pragma once
+
+#include "AssetTools/AssetDefinition_LevelSequence.h"
+
+#include "AssetDefinition_CineAsset.generated.h"
+
+/** Asset definition for a UCineAsset asset */
+UCLASS()
+class UAssetDefinition_CineAsset : public UAssetDefinition_LevelSequence
+{
+	GENERATED_BODY()
+
+protected:
+	// Begin UAssetDefinition Interface
+	virtual FText GetAssetDisplayName() const override;
+	virtual TSoftClassPtr<UObject> GetAssetClass() const override;
+	virtual FLinearColor GetAssetColor() const override;
+	virtual TConstArrayView<FAssetCategoryPath> GetAssetCategories() const override;
+	virtual TArray<FAssetData> PrepareToActivateAssets(const FAssetActivateArgs& ActivateArgs) const override;
+	virtual EAssetCommandResult OpenAssets(const FAssetOpenArgs& OpenArgs) const override;
+	// End UAssetDefinition Interface
+};
diff --git a/Engine/Plugins/VirtualProduction/Shots/Source/ShotsEditor/Private/CineAsset.h b/Engine/Plugins/VirtualProduction/Shots/Source/ShotsEditor/Private/CineAsset.h
new file mode 100644
--- /dev/null
+++ b/Engine/Plugins/VirtualProduction/Shots/Source/ShotsEditor/Private/CineAsset.h
@@ -0,0 +1,22 @@
+// Copyright Epic Games, Inc. All Rights Reserved.
+
+#pragma once
+
+#include "LevelSequence.h"
+
+#include "CineAsset.generated.h"
+
+/** 
+ * A cinematic building block that associates a level sequence with a level
+ */
+UCLASS()
+class UCineAsset : public ULevelSequence
+{
+	GENERATED_BODY()
+
+public:
+	UCineAsset() = default;
+
+	UPROPERTY(EditAnywhere, Category = "Default", meta = (AllowedClasses = "/Script/Engine.World"))
+	FSoftObjectPath Level;
+};
diff --git a/Engine/Plugins/VirtualProduction/Shots/Source/ShotsEditor/Private/CineAssetFactory.cpp b/Engine/Plugins/VirtualProduction/Shots/Source/ShotsEditor/Private/CineAssetFactory.cpp
new file mode 100644
--- /dev/null
+++ b/Engine/Plugins/VirtualProduction/Shots/Source/ShotsEditor/Private/CineAssetFactory.cpp
@@ -0,0 +1,63 @@
+// Copyright Epic Games, Inc. All Rights Reserved.
+
+#include "CineAssetFactory.h"
+#include "CineAsset.h"
+
+#include "Framework/Application/SlateApplication.h"
+#include "Interfaces/IMainFrameModule.h"
+#include "Modules/ModuleManager.h"
+#include "MovieScene.h"
+#include "MovieSceneToolsProjectSettings.h"
+#include "UI/CineAsset/SCreateCineAssetWindow.h"
+
+#define LOCTEXT_NAMESPACE "CineAssetFactory"
+
+UCineAssetFactory::UCineAssetFactory()
+{
+	SupportedClass = UCineAsset::StaticClass();
+	bCreateNew = true;
+	bEditAfterNew = true;
+}
+
+bool UCineAssetFactory::CanCreateNew() const
+{
+	return true;
+}
+
+UObject* UCineAssetFactory::FactoryCreateNew(UClass* Class, UObject* InParent, FName Name, EObjectFlags Flags, UObject* Context, FFeedbackContext* Warn)
+{
+	UCineAsset* NewCineAsset = NewObject<UCineAsset>(InParent, Class, Name, Flags);
+
+	// Do the same setup that ULevelSequence assets do when they are created
+	NewCineAsset->Initialize();
+
+	const UMovieSceneToolsProjectSettings* MovieSceneToolsProjectSettings = GetDefault<UMovieSceneToolsProjectSettings>();
+
+	FFrameRate TickResolution = NewCineAsset->GetMovieScene()->GetTickResolution();
+	NewCineAsset->GetMovieScene()->SetPlaybackRange((MovieSceneToolsProjectSettings->DefaultStartTime * TickResolution).FloorToFrame(), (MovieSceneToolsProjectSettings->DefaultDuration * TickResolution).FloorToFrame().Value);
+
+	// Configure additional properties of the cine asset based on the setup done in the create window
+	NewCineAsset->Level = CineAssetProperties.Level;
+
+	return NewCineAsset;
+}
+
+bool UCineAssetFactory::ConfigureProperties()
+{
+	// Spawn a new window to configure the properties of the cine asset being created
+	IMainFrameModule& MainFrame = FModuleManager::LoadModuleChecked<IMainFrameModule>("MainFrame");
+	TSharedPtr<SWindow>	ParentWindow = MainFrame.GetParentWindow();
+
+	TSharedRef<SCreateCineAssetWindow> CreateCineAssetWindow = SNew(SCreateCineAssetWindow);
+	FSlateApplication::Get().AddModalWindow(CreateCineAssetWindow, ParentWindow);
+
+	if (CreateCineAssetWindow->ShouldCreateNewAsset() == false)
+	{
+		return false;
+	}
+
+	CineAssetProperties = CreateCineAssetWindow->GetCineAssetProperties();
+	return true;
+}
+
+#undef LOCTEXT_NAMESPACE
diff --git a/Engine/Plugins/VirtualProduction/Shots/Source/ShotsEditor/Private/CineAssetFactory.h b/Engine/Plugins/VirtualProduction/Shots/Source/ShotsEditor/Private/CineAssetFactory.h
new file mode 100644
--- /dev/null
+++ b/Engine/Plugins/VirtualProduction/Shots/Source/ShotsEditor/Private/CineAssetFactory.h
@@ -0,0 +1,30 @@
+// Copyright Epic Games, Inc. All Rights Reserved.
+
+#pragma once
+
+#include "Factories/Factory.h"
+#include "CineAssetProperties.h"
+
+#include "CineAssetFactory.generated.h"
+
+/** 
+ * Factory class used to create new UCineAsset objects
+ * Before creating a new Cine Asset, the factory will spawn a new window to configure the properties of the asset that is being created.
+ */
+UCLASS(hidecategories=Object)
+class UCineAssetFactory : public UFactory
+{
+	GENERATED_BODY()
+
+protected:
+	UCineAssetFactory();
+
+	// Begin UFactory Interface
+	virtual UObject* FactoryCreateNew(UClass* Class, UObject* InParent, FName Name, EObjectFlags Flags, UObject* Context, FFeedbackContext* Warn) override;
+	virtual bool CanCreateNew() const override;
+	virtual bool ConfigureProperties() override;
+	// End UFactory Interface
+
+private:
+	FCineAssetProperties CineAssetProperties;
+};
diff --git a/Engine/Plugins/VirtualProduction/Shots/Source/ShotsEditor/Private/CineAssetProperties.h b/Engine/Plugins/VirtualProduction/Shots/Source/ShotsEditor/Private/CineAssetProperties.h
new file mode 100644
--- /dev/null
+++ b/Engine/Plugins/VirtualProduction/Shots/Source/ShotsEditor/Private/CineAssetProperties.h
@@ -0,0 +1,19 @@
+// Copyright Epic Games, Inc. All Rights Reserved.
+
+#pragma once
+
+#include "CineAssetProperties.generated.h"
+
+/**
+ * The properties required to configure a new cine asset after it has been created
+ */
+USTRUCT()
+struct FCineAssetProperties
+{
+	GENERATED_BODY()
+
+public:
+	/** This level will be opened whenever the cine asset is opened. Defaults to the currently open level. This can be modified later. */
+	UPROPERTY(EditAnywhere, Category = "TargetLevel", meta = (AllowedClasses = "/Script/Engine.World"))
+	FSoftObjectPath Level;
+};
diff --git a/Engine/Plugins/VirtualProduction/Shots/Source/ShotsEditor/Private/ShotsEditorModule.cpp b/Engine/Plugins/VirtualProduction/Shots/Source/ShotsEditor/Private/ShotsEditorModule.cpp
--- a/Engine/Plugins/VirtualProduction/Shots/Source/ShotsEditor/Private/ShotsEditorModule.cpp
+++ b/Engine/Plugins/VirtualProduction/Shots/Source/ShotsEditor/Private/ShotsEditorModule.cpp
@@ -2,10 +2,16 @@
 
 #include "Modules/ModuleManager.h"
 
+#include "CineAsset.h"
+#include "ContentBrowserMenuContexts.h"
+#include "Editor.h"
 #include "Framework/Docking/TabManager.h"
 #include "ProductionSettings.h"
 #include "ProductionSettingsCustomization.h"
 #include "PropertyEditorModule.h"
+#include "ShotsEditorStyle.h"
+#include "Subsystems/AssetEditorSubsystem.h"
+#include "ToolMenus.h"
 #include "UI/SProductionWizard.h"
 #include "Widgets/Docking/SDockTab.h"
 #include "WorkspaceMenuStructure.h"
@@ -26,10 +32,13 @@ public:
 
 private:
 	TSharedRef<SDockTab> SpawnProductionWizard(const FSpawnTabArgs& SpawnTabArgs);
+	void RegisterMenus();
 };
 
 void FShotsEditorModule::StartupModule()
 {
+	FShotsEditorStyle::Get();
+
 	FPropertyEditorModule& PropertyEditorModule = FModuleManager::Get().LoadModuleChecked<FPropertyEditorModule>("PropertyEditor");
 
 	PropertyEditorModule.RegisterCustomClassLayout(
@@ -43,6 +52,8 @@ void FShotsEditorModule::StartupModule()
 		.SetTooltipText(LOCTEXT("ProductionWizardTabTooltip", "Open the Production Wizard"))
 		.SetDisplayName(LOCTEXT("ProductionWizardTabTitle", "Production Wizard"))
 		.SetIcon(FSlateIcon(FAppStyle::GetAppStyleSetName(), "ProjectSettings.TabIcon"));
+
+	RegisterMenus();
 }
 
 void FShotsEditorModule::ShutdownModule()
@@ -62,6 +73,44 @@ TSharedRef<SDockTab> FShotsEditorModule::SpawnProductionWizard(const FSpawnTabAr
 		];
 }
 
+void FShotsEditorModule::RegisterMenus()
+{
+	UToolMenus* ToolMenus = UToolMenus::Get();
+
+	// Add a Content Browser context menu action to Cine Assets that allows the asset to be opened in Sequencer without loading its associated map
+	const FName ContextMenuName = TEXT("ContentBrowser.AssetContextMenu.CineAsset");
+	if (UToolMenu* Menu = ToolMenus->ExtendMenu(ContextMenuName))
+	{
+		const FName AssetActionSectionName = TEXT("GetAssetActions");
+		FToolMenuSection& Section = Menu->FindOrAddSection(AssetActionSectionName);
+
+		Section.AddDynamicEntry("CineAssetActions", FNewToolMenuSectionDelegate::CreateLambda([](FToolMenuSection& InSection)
+		{
+			if (UContentBrowserAssetContextMenuContext* Context = InSection.FindContext<UContentBrowserAssetContextMenuContext>())
+			{
+				if ((Context->SelectedAssets.Num() == 1) && Context->SelectedAssets[0].IsInstanceOf(UCineAsset::StaticClass()))
+				{
+					const FAssetData CineAssetData = Context->SelectedAssets[0];
+
+					InSection.AddMenuEntry(
+						"OpenInCurrentMapEntry",
+						LOCTEXT("OpenInCurrentMap", "Open in Current Map"),
+						LOCTEXT("OpenInCurrentMapTooltip", "Opens the level sequence in Sequencer but does not automatically load the associated map"),
+						FSlateIcon(FAppStyle::GetAppStyleSetName(), "LevelEditor.OpenCinematic"),
+						FUIAction(FExecuteAction::CreateLambda([CineAssetData]()
+						{
+							// Note: Normally, when a cine asset is opened in the editor, its associated map is loaded during the PrepareToActiveAssets step.
+							// By directly opening the asset editor (Sequencer) for this asset, we can skip that map loading step. 
+							GEditor->GetEditorSubsystem<UAssetEditorSubsystem>()->OpenEditorForAssets({ CineAssetData.GetAsset() });
+						})),
+						EUserInterfaceActionType::None
+					);
+				}
+			}
+		}));
+	}
+}
+
 IMPLEMENT_MODULE(FShotsEditorModule, ShotsEditor);
 
 #undef LOCTEXT_NAMESPACE
diff --git a/Engine/Plugins/VirtualProduction/Shots/Source/ShotsEditor/Private/UI/CineAsset/SCreateCineAssetWindow.cpp b/Engine/Plugins/VirtualProduction/Shots/Source/ShotsEditor/Private/UI/CineAsset/SCreateCineAssetWindow.cpp
new file mode 100644
--- /dev/null
+++ b/Engine/Plugins/VirtualProduction/Shots/Source/ShotsEditor/Private/UI/CineAsset/SCreateCineAssetWindow.cpp
@@ -0,0 +1,174 @@
+// Copyright Epic Games, Inc. All Rights Reserved.
+
+#include "SCreateCineAssetWindow.h"
+
+#include "Editor.h"
+#include "IStructureDetailsView.h"
+#include "Modules/ModuleManager.h"
+#include "PropertyEditorModule.h"
+#include "ShotsEditorStyle.h"
+#include "Widgets/SBoxPanel.h"
+#include "Widgets/Input/SButton.h"
+#include "Widgets/Layout/SBorder.h"
+#include "Widgets/Layout/SSeparator.h"
+#include "Widgets/Text/STextBlock.h"
+
+#define LOCTEXT_NAMESPACE "SCreateCineAssetWindow"
+
+void SCreateCineAssetWindow::Construct(const FArguments& InArgs)
+{
+	UWorld* CurrentWorld = GEditor->GetEditorWorldContext().World();
+	CineAssetProperties.Level = FSoftObjectPath(CurrentWorld);
+
+	const FVector2D DefaultWindowSize = FVector2D(1400, 750);
+
+	SWindow::Construct(SWindow::FArguments()
+		.Title((LOCTEXT("CineAssetCreationWindowTitle", "Create Cine Asset")))
+		.SizingRule(ESizingRule::UserSized)
+		.ClientSize(DefaultWindowSize)
+		.SupportsMaximize(false)
+		.SupportsMinimize(false)
+		[
+			SNew(SBorder)
+				.BorderImage(FShotsEditorStyle::Get().GetBrush("ProductionWizard.PanelBackground"))
+				[
+					SNew(SVerticalBox)
+
+					+ SVerticalBox::Slot()
+						[
+							SNew(SSplitter)
+								.Orientation(Orient_Horizontal)
+								.PhysicalSplitterHandleSize(2.0f)
+
+								+SSplitter::Slot()
+								.Value(0.5f)
+								.Expose(CineTemplatePanelSlot)
+
+								+ SSplitter::Slot()
+								.Value(0.5f)
+								.Expose(DetailsPanelSlot)
+						]
+
+					+ SVerticalBox::Slot()
+						.AutoHeight()
+						[
+							SNew(SSeparator)
+								.Orientation(Orient_Horizontal)
+								.Thickness(2.0f)
+						]
+
+					+ SVerticalBox::Slot()
+						.AutoHeight()
+						[
+							MakeButtonsPanel()
+						]
+				]
+		]);
+
+	CineTemplatePanelSlot->AttachWidget(MakeCineTemplatePanel());
+	DetailsPanelSlot->AttachWidget(MakeDetailsPanel());
+}
+
+TSharedRef<SWidget> SCreateCineAssetWindow::MakeCineTemplatePanel()
+{
+	// TODO: Fill in this stubbed out panel
+	return SNew(SVerticalBox)
+
+		+ SVerticalBox::Slot()
+		.FillHeight(1.0f)
+		.HAlign(HAlign_Center)
+		[
+			SNew(SHorizontalBox)
+
+			+ SHorizontalBox::Slot()
+				.FillWidth(1.0f)
+				.VAlign(VAlign_Center)
+				[
+					SNew(STextBlock).Text(LOCTEXT("CineTemplatePanel", "CINE TEMPLATE PANEL"))
+				]
+		];
+}
+
+TSharedRef<SWidget> SCreateCineAssetWindow::MakeDetailsPanel()
+{
+	FPropertyEditorModule& PropertyEditor = FModuleManager::Get().LoadModuleChecked<FPropertyEditorModule>(TEXT("PropertyEditor"));
+
+	TSharedRef<FStructOnScope> CineAssetPropertiesStructOnScope = MakeShared<FStructOnScope>(FCineAssetProperties::StaticStruct(), reinterpret_cast<uint8*>(&CineAssetProperties));
+
+	FDetailsViewArgs DetailsViewArgs;
+	DetailsViewArgs.bAllowSearch = false;
+
+	FStructureDetailsViewArgs StructureDetailsViewArgs;
+
+	TSharedPtr<IStructureDetailsView> PropsStructureDetailsView = PropertyEditor.CreateStructureDetailView(DetailsViewArgs, StructureDetailsViewArgs, CineAssetPropertiesStructOnScope);
+	TSharedPtr<SWidget> DetailsWidget = PropsStructureDetailsView->GetWidget();
+
+	return SNew(SVerticalBox)
+
+		+ SVerticalBox::Slot()
+		.FillHeight(1.0f)
+		.HAlign(HAlign_Fill)
+		[
+			SNew(SHorizontalBox)
+
+			+ SHorizontalBox::Slot()
+				.FillWidth(1.0f)
+				.VAlign(VAlign_Top)
+				[
+					DetailsWidget.ToSharedRef()
+				]
+		];
+}
+
+
+TSharedRef<SWidget> SCreateCineAssetWindow::MakeButtonsPanel()
+{
+	return SNew(SBorder)
+		.BorderImage(FAppStyle::Get().GetBrush("Brushes.Panel"))
+		.Padding(16.0f)
+		.HAlign(HAlign_Right)
+		[
+			SNew(SHorizontalBox)
+
+			+ SHorizontalBox::Slot()
+				.MinWidth(118.0f)
+				.MaxWidth(118.0f)
+				.VAlign(VAlign_Center)
+				.Padding(0.0f, 0.0f, 8.0f, 0.0f)
+				[
+					SNew(SButton)
+						.Text(LOCTEXT("CreateAssetButton", "Create Asset"))
+						.ButtonStyle(FAppStyle::Get(), "PrimaryButton")
+						.HAlign(HAlign_Center)
+						.OnClicked(this, &SCreateCineAssetWindow::OnCreateAssetClicked)
+				]
+
+			+ SHorizontalBox::Slot()
+				.MinWidth(118.0f)
+				.MaxWidth(118.0f)
+				.VAlign(VAlign_Center)
+				[
+					SNew(SButton)
+						.Text(LOCTEXT("CancelButton", "Cancel"))
+						.HAlign(HAlign_Center)
+						.OnClicked(this, &SCreateCineAssetWindow::OnCancelClicked)
+				]
+		];
+}
+
+FReply SCreateCineAssetWindow::OnCreateAssetClicked()
+{
+	bShouldCreateNewAsset = true;
+	RequestDestroyWindow();
+	return FReply::Handled();
+}
+
+FReply SCreateCineAssetWindow::OnCancelClicked()
+{
+	bShouldCreateNewAsset = false;
+	RequestDestroyWindow();
+	return FReply::Handled();
+}
+
+
+#undef LOCTEXT_NAMESPACE
diff --git a/Engine/Plugins/VirtualProduction/Shots/Source/ShotsEditor/Private/UI/CineAsset/SCreateCineAssetWindow.h b/Engine/Plugins/VirtualProduction/Shots/Source/ShotsEditor/Private/UI/CineAsset/SCreateCineAssetWindow.h
new file mode 100644
--- /dev/null
+++ b/Engine/Plugins/VirtualProduction/Shots/Source/ShotsEditor/Private/UI/CineAsset/SCreateCineAssetWindow.h
@@ -0,0 +1,56 @@
+// Copyright Epic Games, Inc. All Rights Reserved.
+
+#pragma once
+
+#include "Widgets/SCompoundWidget.h"
+
+#include "CineAssetProperties.h"
+#include "Widgets/Layout/SSplitter.h"
+#include "Widgets/SWindow.h"
+
+/**
+ * A window for configuring the properties of new cine asset before that asset is created
+ */
+class SCreateCineAssetWindow : public SWindow
+{
+public:
+	SCreateCineAssetWindow() = default;
+
+	SLATE_BEGIN_ARGS(SCreateCineAssetWindow) {}
+	SLATE_END_ARGS()
+
+	void Construct(const FArguments& InArgs);
+
+	/** Whether a new asset should be created after closing this window */
+	bool ShouldCreateNewAsset() { return bShouldCreateNewAsset; }
+
+	/** Returns the properties needed to configure for the asset that is about to be created */
+	FCineAssetProperties GetCineAssetProperties() { return CineAssetProperties; }
+
+private:
+	/** Creates the menu on the left side of the production wizard */
+	TSharedRef<SWidget> MakeCineTemplatePanel();
+
+	/** Creates the menu on the left side of the production wizard */
+	TSharedRef<SWidget> MakeDetailsPanel();
+
+	/** Creates the buttons on the bottom of the production wizard */
+	TSharedRef<SWidget> MakeButtonsPanel();
+
+	/** Cycles backwards through the available panels */
+	FReply OnCreateAssetClicked();
+
+	/** Cycles forwards through the available panels */
+	FReply OnCancelClicked();
+
+private:
+	/** Slots for the content panels */
+	SSplitter::FSlot* CineTemplatePanelSlot = nullptr;
+	SSplitter::FSlot* DetailsPanelSlot = nullptr;
+
+	/** Whether a new asset should be created after closing this window */
+	bool bShouldCreateNewAsset = false;
+
+	/** The properties needed to configure for the asset that is about to be created */
+	FCineAssetProperties CineAssetProperties;
+};
diff --git a/Engine/Plugins/VirtualProduction/Shots/Source/ShotsEditor/ShotsEditor.build.cs b/Engine/Plugins/VirtualProduction/Shots/Source/ShotsEditor/ShotsEditor.build.cs
--- a/Engine/Plugins/VirtualProduction/Shots/Source/ShotsEditor/ShotsEditor.build.cs
+++ b/Engine/Plugins/VirtualProduction/Shots/Source/ShotsEditor/ShotsEditor.build.cs
@@ -19,6 +19,7 @@ namespace UnrealBuildTool.Rules
 					"AssetRegistry",
 					"AssetTools",
 					"ClassViewer",
+					"ContentBrowser",
 					"ContentBrowserData",
 					"Core",
 					"CoreUObject",
@@ -40,6 +41,7 @@ namespace UnrealBuildTool.Rules
 					"SlateCore",
 					"SourceControl",
 					"TimeManagement",
+					"ToolMenus",
 					"UnrealEd",
 					"WorkspaceMenuStructure"
 				}
