ShotManagement: Change CineAsset factory to first create a new asset and then hand it to the UI for user configuration

#rnx

diff --git a/Engine/Plugins/VirtualProduction/Shots/Source/ShotsEditor/Private/CineAsset.h b/Engine/Plugins/VirtualProduction/Shots/Source/ShotsEditor/Private/CineAsset.h
--- a/Engine/Plugins/VirtualProduction/Shots/Source/ShotsEditor/Private/CineAsset.h
+++ b/Engine/Plugins/VirtualProduction/Shots/Source/ShotsEditor/Private/CineAsset.h
@@ -9,7 +9,7 @@
 /** 
  * A cinematic building block that associates a level sequence with a level
  */
-UCLASS()
+UCLASS(HideCategories=Animation)
 class UCineAsset : public ULevelSequence
 {
 	GENERATED_BODY()
@@ -17,6 +17,6 @@ class UCineAsset : public ULevelSequence
 public:
 	UCineAsset() = default;
 
-	UPROPERTY(EditAnywhere, Category = "Default", meta = (AllowedClasses = "/Script/Engine.World"))
+	UPROPERTY(EditAnywhere, Category = "TargetLevel", meta = (AllowedClasses = "/Script/Engine.World"))
 	FSoftObjectPath Level;
 };
diff --git a/Engine/Plugins/VirtualProduction/Shots/Source/ShotsEditor/Private/CineAssetFactory.cpp b/Engine/Plugins/VirtualProduction/Shots/Source/ShotsEditor/Private/CineAssetFactory.cpp
--- a/Engine/Plugins/VirtualProduction/Shots/Source/ShotsEditor/Private/CineAssetFactory.cpp
+++ b/Engine/Plugins/VirtualProduction/Shots/Source/ShotsEditor/Private/CineAssetFactory.cpp
@@ -1,7 +1,6 @@
 // Copyright Epic Games, Inc. All Rights Reserved.
 
 #include "CineAssetFactory.h"
-#include "CineAsset.h"
 
 #include "Framework/Application/SlateApplication.h"
 #include "Interfaces/IMainFrameModule.h"
@@ -9,6 +8,7 @@
 #include "MovieScene.h"
 #include "MovieSceneToolsProjectSettings.h"
 #include "UI/CineAsset/SCreateCineAssetWindow.h"
+#include "UObject/Package.h"
 
 #define LOCTEXT_NAMESPACE "CineAssetFactory"
 
@@ -17,6 +17,8 @@ UCineAssetFactory::UCineAssetFactory()
 	SupportedClass = UCineAsset::StaticClass();
 	bCreateNew = true;
 	bEditAfterNew = true;
+
+	CineAssetToConfigure = nullptr;
 }
 
 bool UCineAssetFactory::CanCreateNew() const
@@ -26,20 +28,23 @@ bool UCineAssetFactory::CanCreateNew() const
 
 UObject* UCineAssetFactory::FactoryCreateNew(UClass* Class, UObject* InParent, FName Name, EObjectFlags Flags, UObject* Context, FFeedbackContext* Warn)
 {
-	UCineAsset* NewCineAsset = NewObject<UCineAsset>(InParent, Class, Name, Flags);
-
-	// Do the same setup that ULevelSequence assets do when they are created
-	NewCineAsset->Initialize();
+	if (CineAssetToConfigure)
+	{
+		// Update the name, outer package, and flags of the configuration asset to no longer be transient
+		CineAssetToConfigure->Rename(*Name.ToString(), InParent);
+		CineAssetToConfigure->SetFlags(Flags);
+		CineAssetToConfigure->ClearFlags(RF_Transient);
 
-	const UMovieSceneToolsProjectSettings* MovieSceneToolsProjectSettings = GetDefault<UMovieSceneToolsProjectSettings>();
+		// Do the same setup that ULevelSequence assets do when they are created
+		CineAssetToConfigure->Initialize();
 
-	FFrameRate TickResolution = NewCineAsset->GetMovieScene()->GetTickResolution();
-	NewCineAsset->GetMovieScene()->SetPlaybackRange((MovieSceneToolsProjectSettings->DefaultStartTime * TickResolution).FloorToFrame(), (MovieSceneToolsProjectSettings->DefaultDuration * TickResolution).FloorToFrame().Value);
+		const UMovieSceneToolsProjectSettings* MovieSceneToolsProjectSettings = GetDefault<UMovieSceneToolsProjectSettings>();
 
-	// Configure additional properties of the cine asset based on the setup done in the create window
-	NewCineAsset->Level = CineAssetProperties.Level;
+		FFrameRate TickResolution = CineAssetToConfigure->GetMovieScene()->GetTickResolution();
+		CineAssetToConfigure->GetMovieScene()->SetPlaybackRange((MovieSceneToolsProjectSettings->DefaultStartTime * TickResolution).FloorToFrame(), (MovieSceneToolsProjectSettings->DefaultDuration * TickResolution).FloorToFrame().Value);
+	}
 
-	return NewCineAsset;
+	return CineAssetToConfigure;
 }
 
 bool UCineAssetFactory::ConfigureProperties()
@@ -48,7 +53,11 @@ bool UCineAssetFactory::ConfigureProperties()
 	IMainFrameModule& MainFrame = FModuleManager::LoadModuleChecked<IMainFrameModule>("MainFrame");
 	TSharedPtr<SWindow>	ParentWindow = MainFrame.GetParentWindow();
 
-	TSharedRef<SCreateCineAssetWindow> CreateCineAssetWindow = SNew(SCreateCineAssetWindow);
+	// Create a new transient CineAsset to configure in the UI.
+	// If the configuration is successful, this will turn into the persistent object created by this factory.
+	CineAssetToConfigure = NewObject<UCineAsset>(GetTransientPackage(), NAME_None, RF_Transient);
+
+	TSharedRef<SCreateCineAssetWindow> CreateCineAssetWindow = SNew(SCreateCineAssetWindow, CineAssetToConfigure);
 	FSlateApplication::Get().AddModalWindow(CreateCineAssetWindow, ParentWindow);
 
 	if (CreateCineAssetWindow->ShouldCreateNewAsset() == false)
@@ -56,7 +65,6 @@ bool UCineAssetFactory::ConfigureProperties()
 		return false;
 	}
 
-	CineAssetProperties = CreateCineAssetWindow->GetCineAssetProperties();
 	return true;
 }
 
diff --git a/Engine/Plugins/VirtualProduction/Shots/Source/ShotsEditor/Private/CineAssetFactory.h b/Engine/Plugins/VirtualProduction/Shots/Source/ShotsEditor/Private/CineAssetFactory.h
--- a/Engine/Plugins/VirtualProduction/Shots/Source/ShotsEditor/Private/CineAssetFactory.h
+++ b/Engine/Plugins/VirtualProduction/Shots/Source/ShotsEditor/Private/CineAssetFactory.h
@@ -3,7 +3,7 @@
 #pragma once
 
 #include "Factories/Factory.h"
-#include "CineAssetProperties.h"
+#include "CineAsset.h"
 
 #include "CineAssetFactory.generated.h"
 
@@ -26,5 +26,6 @@ protected:
 	// End UFactory Interface
 
 private:
-	FCineAssetProperties CineAssetProperties;
+	/** Transient CineAsset that will be configured by the user and used as a template for the actual asset created by this factory */
+	UCineAsset* CineAssetToConfigure;
 };
diff --git a/Engine/Plugins/VirtualProduction/Shots/Source/ShotsEditor/Private/CineAssetProperties.h b/Engine/Plugins/VirtualProduction/Shots/Source/ShotsEditor/Private/CineAssetProperties.h
deleted file mode 100644
diff --git a/Engine/Plugins/VirtualProduction/Shots/Source/ShotsEditor/Private/UI/CineAsset/SCreateCineAssetWindow.cpp b/Engine/Plugins/VirtualProduction/Shots/Source/ShotsEditor/Private/UI/CineAsset/SCreateCineAssetWindow.cpp
--- a/Engine/Plugins/VirtualProduction/Shots/Source/ShotsEditor/Private/UI/CineAsset/SCreateCineAssetWindow.cpp
+++ b/Engine/Plugins/VirtualProduction/Shots/Source/ShotsEditor/Private/UI/CineAsset/SCreateCineAssetWindow.cpp
@@ -15,10 +15,12 @@
 
 #define LOCTEXT_NAMESPACE "SCreateCineAssetWindow"
 
-void SCreateCineAssetWindow::Construct(const FArguments& InArgs)
+void SCreateCineAssetWindow::Construct(const FArguments& InArgs, UCineAsset* InCineAsset)
 {
+	CineAssetToConfigure = InCineAsset;
+
 	UWorld* CurrentWorld = GEditor->GetEditorWorldContext().World();
-	CineAssetProperties.Level = FSoftObjectPath(CurrentWorld);
+	CineAssetToConfigure->Level = FSoftObjectPath(CurrentWorld);
 
 	const FVector2D DefaultWindowSize = FVector2D(1400, 750);
 
@@ -66,7 +68,7 @@ void SCreateCineAssetWindow::Construct(const FArguments& InArgs)
 		]);
 
 	CineTemplatePanelSlot->AttachWidget(MakeCineTemplatePanel());
-	DetailsPanelSlot->AttachWidget(MakeDetailsPanel());
+	MakeDetailsPanel();
 }
 
 TSharedRef<SWidget> SCreateCineAssetWindow::MakeCineTemplatePanel()
@@ -89,35 +91,17 @@ TSharedRef<SWidget> SCreateCineAssetWindow::MakeCineTemplatePanel()
 		];
 }
 
-TSharedRef<SWidget> SCreateCineAssetWindow::MakeDetailsPanel()
+void SCreateCineAssetWindow::MakeDetailsPanel()
 {
 	FPropertyEditorModule& PropertyEditor = FModuleManager::Get().LoadModuleChecked<FPropertyEditorModule>(TEXT("PropertyEditor"));
 
-	TSharedRef<FStructOnScope> CineAssetPropertiesStructOnScope = MakeShared<FStructOnScope>(FCineAssetProperties::StaticStruct(), reinterpret_cast<uint8*>(&CineAssetProperties));
-
 	FDetailsViewArgs DetailsViewArgs;
 	DetailsViewArgs.bAllowSearch = false;
+	DetailsViewArgs.bHideSelectionTip = true;
 
-	FStructureDetailsViewArgs StructureDetailsViewArgs;
-
-	TSharedPtr<IStructureDetailsView> PropsStructureDetailsView = PropertyEditor.CreateStructureDetailView(DetailsViewArgs, StructureDetailsViewArgs, CineAssetPropertiesStructOnScope);
-	TSharedPtr<SWidget> DetailsWidget = PropsStructureDetailsView->GetWidget();
-
-	return SNew(SVerticalBox)
-
-		+ SVerticalBox::Slot()
-		.FillHeight(1.0f)
-		.HAlign(HAlign_Fill)
-		[
-			SNew(SHorizontalBox)
-
-			+ SHorizontalBox::Slot()
-				.FillWidth(1.0f)
-				.VAlign(VAlign_Top)
-				[
-					DetailsWidget.ToSharedRef()
-				]
-		];
+	TSharedRef<IDetailsView> DetailsView = PropertyEditor.CreateDetailView(DetailsViewArgs);
+	DetailsView->SetObject(CineAssetToConfigure, true);
+	DetailsPanelSlot->AttachWidget(DetailsView);
 }
 
 
diff --git a/Engine/Plugins/VirtualProduction/Shots/Source/ShotsEditor/Private/UI/CineAsset/SCreateCineAssetWindow.h b/Engine/Plugins/VirtualProduction/Shots/Source/ShotsEditor/Private/UI/CineAsset/SCreateCineAssetWindow.h
--- a/Engine/Plugins/VirtualProduction/Shots/Source/ShotsEditor/Private/UI/CineAsset/SCreateCineAssetWindow.h
+++ b/Engine/Plugins/VirtualProduction/Shots/Source/ShotsEditor/Private/UI/CineAsset/SCreateCineAssetWindow.h
@@ -4,7 +4,7 @@
 
 #include "Widgets/SCompoundWidget.h"
 
-#include "CineAssetProperties.h"
+#include "CineAsset.h"
 #include "Widgets/Layout/SSplitter.h"
 #include "Widgets/SWindow.h"
 
@@ -19,20 +19,17 @@ public:
 	SLATE_BEGIN_ARGS(SCreateCineAssetWindow) {}
 	SLATE_END_ARGS()
 
-	void Construct(const FArguments& InArgs);
+	void Construct(const FArguments& InArgs, UCineAsset* InCineAsset);
 
 	/** Whether a new asset should be created after closing this window */
 	bool ShouldCreateNewAsset() { return bShouldCreateNewAsset; }
 
-	/** Returns the properties needed to configure for the asset that is about to be created */
-	FCineAssetProperties GetCineAssetProperties() { return CineAssetProperties; }
-
 private:
 	/** Creates the menu on the left side of the production wizard */
 	TSharedRef<SWidget> MakeCineTemplatePanel();
 
 	/** Creates the menu on the left side of the production wizard */
-	TSharedRef<SWidget> MakeDetailsPanel();
+	void MakeDetailsPanel();
 
 	/** Creates the buttons on the bottom of the production wizard */
 	TSharedRef<SWidget> MakeButtonsPanel();
@@ -51,6 +48,6 @@ private:
 	/** Whether a new asset should be created after closing this window */
 	bool bShouldCreateNewAsset = false;
 
-	/** The properties needed to configure for the asset that is about to be created */
-	FCineAssetProperties CineAssetProperties;
+	/** Transient object used only by this UI to configure the properties of the new asset that will get created by the Factory */
+	UCineAsset* CineAssetToConfigure = nullptr;
 };
