ShotManagement: Rename CineAsset to CineAssembly

#rnx

diff --git a/Engine/Plugins/VirtualProduction/Shots/Source/ShotsEditor/Private/AssetDefinition_CineAssembly.cpp b/Engine/Plugins/VirtualProduction/Shots/Source/ShotsEditor/Private/AssetDefinition_CineAssembly.cpp
new file mode 100644
--- /dev/null
+++ b/Engine/Plugins/VirtualProduction/Shots/Source/ShotsEditor/Private/AssetDefinition_CineAssembly.cpp
@@ -0,0 +1,90 @@
+// Copyright Epic Games, Inc. All Rights Reserved.
+
+#include "AssetDefinition_CineAssembly.h"
+#include "CineAssembly.h"
+
+#include "Editor.h"
+#include "FileHelpers.h"
+
+#define LOCTEXT_NAMESPACE "AssetTypeActions"
+
+DEFINE_LOG_CATEGORY_STATIC(LogCineAssemblyDefinition, Log, All)
+
+FText UAssetDefinition_CineAssembly::GetAssetDisplayName() const
+{
+	return LOCTEXT("AssetTypeActions_CineAssembly", "Cine Assembly");
+}
+
+TSoftClassPtr<> UAssetDefinition_CineAssembly::GetAssetClass() const
+{
+	return UCineAssembly::StaticClass();
+}
+
+FLinearColor UAssetDefinition_CineAssembly::GetAssetColor() const
+{
+	return FColor(172, 35, 85);
+}
+
+TConstArrayView<FAssetCategoryPath> UAssetDefinition_CineAssembly::GetAssetCategories() const
+{
+	static const FAssetCategoryPath Categories[] = { EAssetCategoryPaths::Cinematics };
+	return Categories;
+}
+
+EAssetCommandResult UAssetDefinition_CineAssembly::OpenAssets(const FAssetOpenArgs& OpenArgs) const
+{
+	return Super::OpenAssets(OpenArgs);
+}
+
+TArray<FAssetData> UAssetDefinition_CineAssembly::PrepareToActivateAssets(const FAssetActivateArgs& ActivateArgs) const
+{
+	TArray<FAssetData> AssetsToOpen;
+
+	// We only support opening one asset at a time
+	if (ActivateArgs.Assets.Num() == 1)
+	{
+		const FAssetData& CineAssemblyData = ActivateArgs.Assets[0];
+		AssetsToOpen.Add(CineAssemblyData);
+
+		if (UCineAssembly* CineAssembly = Cast<UCineAssembly>(CineAssemblyData.GetAsset()))
+		{
+			if (CineAssembly->Level.IsValid())
+			{
+				if (UWorld* WorldToOpen = Cast<UWorld>(CineAssembly->Level.TryLoad()))
+				{
+					UWorld* CurrentWorld = GEditor->GetEditorWorldContext().World();
+					if (CurrentWorld != WorldToOpen)
+					{
+						// Prompt the user to save their unsaved changes to the current level before loading the level associated with this asset
+						constexpr bool bPromptUserToSave = true;
+						constexpr bool bSaveMapPackages = true;
+						constexpr bool bSaveContentPackages = false;
+						if (FEditorFileUtils::SaveDirtyPackages(bPromptUserToSave, bSaveMapPackages, bSaveContentPackages))
+						{
+							if (!WorldToOpen->GetPackage()->HasAnyPackageFlags(PKG_NewlyCreated))
+							{
+								const FString FileToOpen = FPackageName::LongPackageNameToFilename(WorldToOpen->GetOutermost()->GetName(), FPackageName::GetMapPackageExtension());
+								const bool bLoadAsTemplate = false;
+								const bool bShowProgress = true;
+								FEditorFileUtils::LoadMap(FileToOpen, bLoadAsTemplate, bShowProgress);
+							}
+						}
+						else
+						{
+							// If the user canceled out of the prompt to save the current level, then do not try to open the asset
+							AssetsToOpen.Empty();
+						}
+					}
+				}
+				else
+				{
+					UE_LOG(LogCineAssemblyDefinition, Error, TEXT("Failed to load %s while opening %s"), *CineAssembly->Level.ToString(), *CineAssembly->GetFName().ToString());
+				}
+			}
+		}
+	}
+
+	return AssetsToOpen;
+}
+
+#undef LOCTEXT_NAMESPACE
diff --git a/Engine/Plugins/VirtualProduction/Shots/Source/ShotsEditor/Private/AssetDefinition_CineAssembly.h b/Engine/Plugins/VirtualProduction/Shots/Source/ShotsEditor/Private/AssetDefinition_CineAssembly.h
new file mode 100644
--- /dev/null
+++ b/Engine/Plugins/VirtualProduction/Shots/Source/ShotsEditor/Private/AssetDefinition_CineAssembly.h
@@ -0,0 +1,24 @@
+// Copyright Epic Games, Inc. All Rights Reserved.
+
+#pragma once
+
+#include "AssetTools/AssetDefinition_LevelSequence.h"
+
+#include "AssetDefinition_CineAssembly.generated.h"
+
+/** Asset definition for a UCineAssembly asset */
+UCLASS()
+class UAssetDefinition_CineAssembly : public UAssetDefinition_LevelSequence
+{
+	GENERATED_BODY()
+
+protected:
+	// Begin UAssetDefinition Interface
+	virtual FText GetAssetDisplayName() const override;
+	virtual TSoftClassPtr<UObject> GetAssetClass() const override;
+	virtual FLinearColor GetAssetColor() const override;
+	virtual TConstArrayView<FAssetCategoryPath> GetAssetCategories() const override;
+	virtual TArray<FAssetData> PrepareToActivateAssets(const FAssetActivateArgs& ActivateArgs) const override;
+	virtual EAssetCommandResult OpenAssets(const FAssetOpenArgs& OpenArgs) const override;
+	// End UAssetDefinition Interface
+};
diff --git a/Engine/Plugins/VirtualProduction/Shots/Source/ShotsEditor/Private/AssetDefinition_CineAsset.cpp b/Engine/Plugins/VirtualProduction/Shots/Source/ShotsEditor/Private/AssetDefinition_CineAsset.cpp
deleted file mode 100644
diff --git a/Engine/Plugins/VirtualProduction/Shots/Source/ShotsEditor/Private/AssetDefinition_CineAsset.h b/Engine/Plugins/VirtualProduction/Shots/Source/ShotsEditor/Private/AssetDefinition_CineAsset.h
deleted file mode 100644
diff --git a/Engine/Plugins/VirtualProduction/Shots/Source/ShotsEditor/Private/CineAssembly.h b/Engine/Plugins/VirtualProduction/Shots/Source/ShotsEditor/Private/CineAssembly.h
new file mode 100644
--- /dev/null
+++ b/Engine/Plugins/VirtualProduction/Shots/Source/ShotsEditor/Private/CineAssembly.h
@@ -0,0 +1,22 @@
+// Copyright Epic Games, Inc. All Rights Reserved.
+
+#pragma once
+
+#include "LevelSequence.h"
+
+#include "CineAssembly.generated.h"
+
+/** 
+ * A cinematic building block that associates a level sequence with a level
+ */
+UCLASS(HideCategories=Animation)
+class UCineAssembly : public ULevelSequence
+{
+	GENERATED_BODY()
+
+public:
+	UCineAssembly() = default;
+
+	UPROPERTY(EditAnywhere, Category = "TargetLevel", meta = (AllowedClasses = "/Script/Engine.World"))
+	FSoftObjectPath Level;
+};
diff --git a/Engine/Plugins/VirtualProduction/Shots/Source/ShotsEditor/Private/CineAssemblyFactory.cpp b/Engine/Plugins/VirtualProduction/Shots/Source/ShotsEditor/Private/CineAssemblyFactory.cpp
new file mode 100644
--- /dev/null
+++ b/Engine/Plugins/VirtualProduction/Shots/Source/ShotsEditor/Private/CineAssemblyFactory.cpp
@@ -0,0 +1,71 @@
+// Copyright Epic Games, Inc. All Rights Reserved.
+
+#include "CineAssemblyFactory.h"
+
+#include "Framework/Application/SlateApplication.h"
+#include "Interfaces/IMainFrameModule.h"
+#include "Modules/ModuleManager.h"
+#include "MovieScene.h"
+#include "MovieSceneToolsProjectSettings.h"
+#include "UI/CineAssembly/SCineAssemblyConfigWindow.h"
+#include "UObject/Package.h"
+
+#define LOCTEXT_NAMESPACE "CineAssemblyFactory"
+
+UCineAssemblyFactory::UCineAssemblyFactory()
+{
+	SupportedClass = UCineAssembly::StaticClass();
+	bCreateNew = true;
+	bEditAfterNew = true;
+
+	CineAssemblyToConfigure = nullptr;
+}
+
+bool UCineAssemblyFactory::CanCreateNew() const
+{
+	return true;
+}
+
+UObject* UCineAssemblyFactory::FactoryCreateNew(UClass* Class, UObject* InParent, FName Name, EObjectFlags Flags, UObject* Context, FFeedbackContext* Warn)
+{
+	if (CineAssemblyToConfigure)
+	{
+		// Update the name, outer package, and flags of the configuration asset to no longer be transient
+		CineAssemblyToConfigure->Rename(*Name.ToString(), InParent);
+		CineAssemblyToConfigure->SetFlags(Flags);
+		CineAssemblyToConfigure->ClearFlags(RF_Transient);
+
+		// Do the same setup that ULevelSequence assets do when they are created
+		CineAssemblyToConfigure->Initialize();
+
+		const UMovieSceneToolsProjectSettings* MovieSceneToolsProjectSettings = GetDefault<UMovieSceneToolsProjectSettings>();
+
+		FFrameRate TickResolution = CineAssemblyToConfigure->GetMovieScene()->GetTickResolution();
+		CineAssemblyToConfigure->GetMovieScene()->SetPlaybackRange((MovieSceneToolsProjectSettings->DefaultStartTime * TickResolution).FloorToFrame(), (MovieSceneToolsProjectSettings->DefaultDuration * TickResolution).FloorToFrame().Value);
+	}
+
+	return CineAssemblyToConfigure;
+}
+
+bool UCineAssemblyFactory::ConfigureProperties()
+{
+	// Spawn a new window to configure the properties of the asset being created
+	IMainFrameModule& MainFrame = FModuleManager::LoadModuleChecked<IMainFrameModule>("MainFrame");
+	TSharedPtr<SWindow>	ParentWindow = MainFrame.GetParentWindow();
+
+	// Create a new transient CineAssembly to configure in the UI.
+	// If the configuration is successful, this will turn into the persistent object created by this factory.
+	CineAssemblyToConfigure = NewObject<UCineAssembly>(GetTransientPackage(), NAME_None, RF_Transient);
+
+	TSharedRef<SCineAssemblyConfigWindow> CineAssemblyConfigWindow = SNew(SCineAssemblyConfigWindow, CineAssemblyToConfigure);
+	FSlateApplication::Get().AddModalWindow(CineAssemblyConfigWindow, ParentWindow);
+
+	if (CineAssemblyConfigWindow->ShouldCreateNewAsset() == false)
+	{
+		return false;
+	}
+
+	return true;
+}
+
+#undef LOCTEXT_NAMESPACE
diff --git a/Engine/Plugins/VirtualProduction/Shots/Source/ShotsEditor/Private/CineAssemblyFactory.h b/Engine/Plugins/VirtualProduction/Shots/Source/ShotsEditor/Private/CineAssemblyFactory.h
new file mode 100644
--- /dev/null
+++ b/Engine/Plugins/VirtualProduction/Shots/Source/ShotsEditor/Private/CineAssemblyFactory.h
@@ -0,0 +1,31 @@
+// Copyright Epic Games, Inc. All Rights Reserved.
+
+#pragma once
+
+#include "Factories/Factory.h"
+#include "CineAssembly.h"
+
+#include "CineAssemblyFactory.generated.h"
+
+/** 
+ * Factory class used to create new UCineAssembly objects
+ * Before creating a new Cine Assembly, the factory will spawn a new window to configure the properties of the asset that is being created.
+ */
+UCLASS(hidecategories=Object)
+class UCineAssemblyFactory : public UFactory
+{
+	GENERATED_BODY()
+
+protected:
+	UCineAssemblyFactory();
+
+	// Begin UFactory Interface
+	virtual UObject* FactoryCreateNew(UClass* Class, UObject* InParent, FName Name, EObjectFlags Flags, UObject* Context, FFeedbackContext* Warn) override;
+	virtual bool CanCreateNew() const override;
+	virtual bool ConfigureProperties() override;
+	// End UFactory Interface
+
+private:
+	/** Transient CineAssembly that will be configured by the user and used as a template for the actual asset created by this factory */
+	UCineAssembly* CineAssemblyToConfigure;
+};
diff --git a/Engine/Plugins/VirtualProduction/Shots/Source/ShotsEditor/Private/CineAsset.h b/Engine/Plugins/VirtualProduction/Shots/Source/ShotsEditor/Private/CineAsset.h
deleted file mode 100644
diff --git a/Engine/Plugins/VirtualProduction/Shots/Source/ShotsEditor/Private/CineAssetFactory.cpp b/Engine/Plugins/VirtualProduction/Shots/Source/ShotsEditor/Private/CineAssetFactory.cpp
deleted file mode 100644
diff --git a/Engine/Plugins/VirtualProduction/Shots/Source/ShotsEditor/Private/CineAssetFactory.h b/Engine/Plugins/VirtualProduction/Shots/Source/ShotsEditor/Private/CineAssetFactory.h
deleted file mode 100644
diff --git a/Engine/Plugins/VirtualProduction/Shots/Source/ShotsEditor/Private/ShotsEditorModule.cpp b/Engine/Plugins/VirtualProduction/Shots/Source/ShotsEditor/Private/ShotsEditorModule.cpp
--- a/Engine/Plugins/VirtualProduction/Shots/Source/ShotsEditor/Private/ShotsEditorModule.cpp
+++ b/Engine/Plugins/VirtualProduction/Shots/Source/ShotsEditor/Private/ShotsEditorModule.cpp
@@ -2,7 +2,7 @@
 
 #include "Modules/ModuleManager.h"
 
-#include "CineAsset.h"
+#include "CineAssembly.h"
 #include "ContentBrowserMenuContexts.h"
 #include "Editor.h"
 #include "Framework/Docking/TabManager.h"
@@ -77,31 +77,31 @@ void FShotsEditorModule::RegisterMenus()
 {
 	UToolMenus* ToolMenus = UToolMenus::Get();
 
-	// Add a Content Browser context menu action to Cine Assets that allows the asset to be opened in Sequencer without loading its associated map
-	const FName ContextMenuName = TEXT("ContentBrowser.AssetContextMenu.CineAsset");
+	// Add a Content Browser context menu action to Cine Assembly assets that allows the asset to be opened in Sequencer without loading its associated map
+	const FName ContextMenuName = TEXT("ContentBrowser.AssetContextMenu.CineAssembly");
 	if (UToolMenu* Menu = ToolMenus->ExtendMenu(ContextMenuName))
 	{
 		const FName AssetActionSectionName = TEXT("GetAssetActions");
 		FToolMenuSection& Section = Menu->FindOrAddSection(AssetActionSectionName);
 
-		Section.AddDynamicEntry("CineAssetActions", FNewToolMenuSectionDelegate::CreateLambda([](FToolMenuSection& InSection)
+		Section.AddDynamicEntry("CineAssemblyActions", FNewToolMenuSectionDelegate::CreateLambda([](FToolMenuSection& InSection)
 		{
 			if (UContentBrowserAssetContextMenuContext* Context = InSection.FindContext<UContentBrowserAssetContextMenuContext>())
 			{
-				if ((Context->SelectedAssets.Num() == 1) && Context->SelectedAssets[0].IsInstanceOf(UCineAsset::StaticClass()))
+				if ((Context->SelectedAssets.Num() == 1) && Context->SelectedAssets[0].IsInstanceOf(UCineAssembly::StaticClass()))
 				{
-					const FAssetData CineAssetData = Context->SelectedAssets[0];
+					const FAssetData CineAssemblyData = Context->SelectedAssets[0];
 
 					InSection.AddMenuEntry(
 						"OpenInCurrentMapEntry",
 						LOCTEXT("OpenInCurrentMap", "Open in Current Map"),
 						LOCTEXT("OpenInCurrentMapTooltip", "Opens the level sequence in Sequencer but does not automatically load the associated map"),
 						FSlateIcon(FAppStyle::GetAppStyleSetName(), "LevelEditor.OpenCinematic"),
-						FUIAction(FExecuteAction::CreateLambda([CineAssetData]()
+						FUIAction(FExecuteAction::CreateLambda([CineAssemblyData]()
 						{
-							// Note: Normally, when a cine asset is opened in the editor, its associated map is loaded during the PrepareToActiveAssets step.
+							// Note: Normally, when a cine assembly is opened in the editor, its associated map is loaded during the PrepareToActiveAssets step.
 							// By directly opening the asset editor (Sequencer) for this asset, we can skip that map loading step. 
-							GEditor->GetEditorSubsystem<UAssetEditorSubsystem>()->OpenEditorForAssets({ CineAssetData.GetAsset() });
+							GEditor->GetEditorSubsystem<UAssetEditorSubsystem>()->OpenEditorForAssets({ CineAssemblyData.GetAsset() });
 						})),
 						EUserInterfaceActionType::None
 					);
diff --git a/Engine/Plugins/VirtualProduction/Shots/Source/ShotsEditor/Private/UI/CineAssembly/SCineAssemblyConfigWindow.cpp b/Engine/Plugins/VirtualProduction/Shots/Source/ShotsEditor/Private/UI/CineAssembly/SCineAssemblyConfigWindow.cpp
new file mode 100644
--- /dev/null
+++ b/Engine/Plugins/VirtualProduction/Shots/Source/ShotsEditor/Private/UI/CineAssembly/SCineAssemblyConfigWindow.cpp
@@ -0,0 +1,158 @@
+// Copyright Epic Games, Inc. All Rights Reserved.
+
+#include "SCineAssemblyConfigWindow.h"
+
+#include "Editor.h"
+#include "IStructureDetailsView.h"
+#include "Modules/ModuleManager.h"
+#include "PropertyEditorModule.h"
+#include "ShotsEditorStyle.h"
+#include "Widgets/SBoxPanel.h"
+#include "Widgets/Input/SButton.h"
+#include "Widgets/Layout/SBorder.h"
+#include "Widgets/Layout/SSeparator.h"
+#include "Widgets/Text/STextBlock.h"
+
+#define LOCTEXT_NAMESPACE "SCineAssemblyConfigWindow"
+
+void SCineAssemblyConfigWindow::Construct(const FArguments& InArgs, UCineAssembly* InCineAssembly)
+{
+	CineAssemblyToConfigure = InCineAssembly;
+
+	UWorld* CurrentWorld = GEditor->GetEditorWorldContext().World();
+	CineAssemblyToConfigure->Level = FSoftObjectPath(CurrentWorld);
+
+	const FVector2D DefaultWindowSize = FVector2D(1400, 750);
+
+	SWindow::Construct(SWindow::FArguments()
+		.Title((LOCTEXT("CineAssemblyCreationWindowTitle", "Create Cine Assembly")))
+		.SizingRule(ESizingRule::UserSized)
+		.ClientSize(DefaultWindowSize)
+		.SupportsMaximize(false)
+		.SupportsMinimize(false)
+		[
+			SNew(SBorder)
+				.BorderImage(FShotsEditorStyle::Get().GetBrush("ProductionWizard.PanelBackground"))
+				[
+					SNew(SVerticalBox)
+
+					+ SVerticalBox::Slot()
+						[
+							SNew(SSplitter)
+								.Orientation(Orient_Horizontal)
+								.PhysicalSplitterHandleSize(2.0f)
+
+								+SSplitter::Slot()
+								.Value(0.5f)
+								.Expose(CineTemplatePanelSlot)
+
+								+ SSplitter::Slot()
+								.Value(0.5f)
+								.Expose(DetailsPanelSlot)
+						]
+
+					+ SVerticalBox::Slot()
+						.AutoHeight()
+						[
+							SNew(SSeparator)
+								.Orientation(Orient_Horizontal)
+								.Thickness(2.0f)
+						]
+
+					+ SVerticalBox::Slot()
+						.AutoHeight()
+						[
+							MakeButtonsPanel()
+						]
+				]
+		]);
+
+	CineTemplatePanelSlot->AttachWidget(MakeCineTemplatePanel());
+	MakeDetailsPanel();
+}
+
+TSharedRef<SWidget> SCineAssemblyConfigWindow::MakeCineTemplatePanel()
+{
+	// TODO: Fill in this stubbed out panel
+	return SNew(SVerticalBox)
+
+		+ SVerticalBox::Slot()
+		.FillHeight(1.0f)
+		.HAlign(HAlign_Center)
+		[
+			SNew(SHorizontalBox)
+
+			+ SHorizontalBox::Slot()
+				.FillWidth(1.0f)
+				.VAlign(VAlign_Center)
+				[
+					SNew(STextBlock).Text(LOCTEXT("CineTemplatePanel", "CINE TEMPLATE PANEL"))
+				]
+		];
+}
+
+void SCineAssemblyConfigWindow::MakeDetailsPanel()
+{
+	FPropertyEditorModule& PropertyEditor = FModuleManager::Get().LoadModuleChecked<FPropertyEditorModule>(TEXT("PropertyEditor"));
+
+	FDetailsViewArgs DetailsViewArgs;
+	DetailsViewArgs.bAllowSearch = false;
+	DetailsViewArgs.bHideSelectionTip = true;
+
+	TSharedRef<IDetailsView> DetailsView = PropertyEditor.CreateDetailView(DetailsViewArgs);
+	DetailsView->SetObject(CineAssemblyToConfigure, true);
+	DetailsPanelSlot->AttachWidget(DetailsView);
+}
+
+
+TSharedRef<SWidget> SCineAssemblyConfigWindow::MakeButtonsPanel()
+{
+	return SNew(SBorder)
+		.BorderImage(FAppStyle::Get().GetBrush("Brushes.Panel"))
+		.Padding(16.0f)
+		.HAlign(HAlign_Right)
+		[
+			SNew(SHorizontalBox)
+
+			+ SHorizontalBox::Slot()
+				.MinWidth(118.0f)
+				.MaxWidth(118.0f)
+				.VAlign(VAlign_Center)
+				.Padding(0.0f, 0.0f, 8.0f, 0.0f)
+				[
+					SNew(SButton)
+						.Text(LOCTEXT("CreateAssetButton", "Create Asset"))
+						.ButtonStyle(FAppStyle::Get(), "PrimaryButton")
+						.HAlign(HAlign_Center)
+						.OnClicked(this, &SCineAssemblyConfigWindow::OnCreateAssetClicked)
+				]
+
+			+ SHorizontalBox::Slot()
+				.MinWidth(118.0f)
+				.MaxWidth(118.0f)
+				.VAlign(VAlign_Center)
+				[
+					SNew(SButton)
+						.Text(LOCTEXT("CancelButton", "Cancel"))
+						.HAlign(HAlign_Center)
+						.OnClicked(this, &SCineAssemblyConfigWindow::OnCancelClicked)
+				]
+		];
+}
+
+FReply SCineAssemblyConfigWindow::OnCreateAssetClicked()
+{
+	bShouldCreateNewAsset = true;
+	RequestDestroyWindow();
+	return FReply::Handled();
+}
+
+FReply SCineAssemblyConfigWindow::OnCancelClicked()
+{
+	bShouldCreateNewAsset = false;
+	RequestDestroyWindow();
+	return FReply::Handled();
+}
+
+
+#undef LOCTEXT_NAMESPACE
diff --git a/Engine/Plugins/VirtualProduction/Shots/Source/ShotsEditor/Private/UI/CineAssembly/SCineAssemblyConfigWindow.h b/Engine/Plugins/VirtualProduction/Shots/Source/ShotsEditor/Private/UI/CineAssembly/SCineAssemblyConfigWindow.h
new file mode 100644
--- /dev/null
+++ b/Engine/Plugins/VirtualProduction/Shots/Source/ShotsEditor/Private/UI/CineAssembly/SCineAssemblyConfigWindow.h
@@ -0,0 +1,53 @@
+// Copyright Epic Games, Inc. All Rights Reserved.
+
+#pragma once
+
+#include "Widgets/SCompoundWidget.h"
+
+#include "CineAssembly.h"
+#include "Widgets/Layout/SSplitter.h"
+#include "Widgets/SWindow.h"
+
+/**
+ * A window for configuring the properties of a UCineAssembly asset
+ */
+class SCineAssemblyConfigWindow : public SWindow
+{
+public:
+	SCineAssemblyConfigWindow() = default;
+
+	SLATE_BEGIN_ARGS(SCineAssemblyConfigWindow) {}
+	SLATE_END_ARGS()
+
+	void Construct(const FArguments& InArgs, UCineAssembly* InCineAssembly);
+
+	/** Whether a new asset should be created after closing this window */
+	bool ShouldCreateNewAsset() { return bShouldCreateNewAsset; }
+
+private:
+	/** Creates the menu on the left side of the production wizard */
+	TSharedRef<SWidget> MakeCineTemplatePanel();
+
+	/** Creates the menu on the left side of the production wizard */
+	void MakeDetailsPanel();
+
+	/** Creates the buttons on the bottom of the production wizard */
+	TSharedRef<SWidget> MakeButtonsPanel();
+
+	/** Cycles backwards through the available panels */
+	FReply OnCreateAssetClicked();
+
+	/** Cycles forwards through the available panels */
+	FReply OnCancelClicked();
+
+private:
+	/** Slots for the content panels */
+	SSplitter::FSlot* CineTemplatePanelSlot = nullptr;
+	SSplitter::FSlot* DetailsPanelSlot = nullptr;
+
+	/** Whether a new asset should be created after closing this window */
+	bool bShouldCreateNewAsset = false;
+
+	/** Transient object used only by this UI to configure the properties of the new asset that will get created by the Factory */
+	UCineAssembly* CineAssemblyToConfigure = nullptr;
+};
diff --git a/Engine/Plugins/VirtualProduction/Shots/Source/ShotsEditor/Private/UI/CineAsset/SCreateCineAssetWindow.cpp b/Engine/Plugins/VirtualProduction/Shots/Source/ShotsEditor/Private/UI/CineAsset/SCreateCineAssetWindow.cpp
deleted file mode 100644
diff --git a/Engine/Plugins/VirtualProduction/Shots/Source/ShotsEditor/Private/UI/CineAsset/SCreateCineAssetWindow.h b/Engine/Plugins/VirtualProduction/Shots/Source/ShotsEditor/Private/UI/CineAsset/SCreateCineAssetWindow.h
deleted file mode 100644
