ShotManagement: Add Edit Properties context option to Cine Assembly assets which opens the asset in the same config window that was used to create it.
#rnx

diff --git a/Engine/Plugins/VirtualProduction/Shots/Source/ShotsEditor/Private/CineAssemblyFactory.cpp b/Engine/Plugins/VirtualProduction/Shots/Source/ShotsEditor/Private/CineAssemblyFactory.cpp
--- a/Engine/Plugins/VirtualProduction/Shots/Source/ShotsEditor/Private/CineAssemblyFactory.cpp
+++ b/Engine/Plugins/VirtualProduction/Shots/Source/ShotsEditor/Private/CineAssemblyFactory.cpp
@@ -57,7 +57,7 @@ bool UCineAssemblyFactory::ConfigureProperties()
 	// If the configuration is successful, this will turn into the persistent object created by this factory.
 	CineAssemblyToConfigure = NewObject<UCineAssembly>(GetTransientPackage(), NAME_None, RF_Transient);
 
-	TSharedRef<SCineAssemblyConfigWindow> CineAssemblyConfigWindow = SNew(SCineAssemblyConfigWindow, CineAssemblyToConfigure);
+	TSharedRef<SCineAssemblyConfigWindow> CineAssemblyConfigWindow = SNew(SCineAssemblyConfigWindow, CineAssemblyToConfigure, ECineAssemblyConfigMode::CreateNew);
 	FSlateApplication::Get().AddModalWindow(CineAssemblyConfigWindow, ParentWindow);
 
 	if (CineAssemblyConfigWindow->ShouldCreateNewAsset() == false)
diff --git a/Engine/Plugins/VirtualProduction/Shots/Source/ShotsEditor/Private/ShotsEditorModule.cpp b/Engine/Plugins/VirtualProduction/Shots/Source/ShotsEditor/Private/ShotsEditorModule.cpp
--- a/Engine/Plugins/VirtualProduction/Shots/Source/ShotsEditor/Private/ShotsEditorModule.cpp
+++ b/Engine/Plugins/VirtualProduction/Shots/Source/ShotsEditor/Private/ShotsEditorModule.cpp
@@ -5,7 +5,9 @@
 #include "CineAssembly.h"
 #include "ContentBrowserMenuContexts.h"
 #include "Editor.h"
+#include "Framework/Application/SlateApplication.h"
 #include "Framework/Docking/TabManager.h"
+#include "Interfaces/IMainFrameModule.h"
 #include "ProductionSettings.h"
 #include "ProductionSettingsCustomization.h"
 #include "PropertyEditorModule.h"
@@ -13,6 +15,7 @@
 #include "Subsystems/AssetEditorSubsystem.h"
 #include "ToolMenus.h"
 #include "UI/SProductionWizard.h"
+#include "UI/CineAssembly/SCineAssemblyConfigWindow.h"
 #include "Widgets/Docking/SDockTab.h"
 #include "WorkspaceMenuStructure.h"
 #include "WorkspaceMenuStructureModule.h"
@@ -105,6 +108,23 @@ void FShotsEditorModule::RegisterMenus()
 						})),
 						EUserInterfaceActionType::None
 					);
+
+					InSection.AddMenuEntry(
+						"OpenInAssetEditor",
+						LOCTEXT("OpenInAssetEditor", "Edit Properties"),
+						LOCTEXT("OpenInAssetEditorTooltip", "Edit Properties"),
+						FSlateIcon(FAppStyle::GetAppStyleSetName(), "Icons.Edit"),
+						FUIAction(FExecuteAction::CreateLambda([CineAssemblyData]()
+						{
+							// Spawn a new window to configure the properties of the asset being created
+							IMainFrameModule& MainFrame = FModuleManager::LoadModuleChecked<IMainFrameModule>("MainFrame");
+							TSharedPtr<SWindow>	ParentWindow = MainFrame.GetParentWindow();
+
+							TSharedRef<SCineAssemblyConfigWindow> CineAssemblyConfigWindow = SNew(SCineAssemblyConfigWindow, Cast<UCineAssembly>(CineAssemblyData.GetAsset()), ECineAssemblyConfigMode::Edit);
+							FSlateApplication::Get().AddModalWindow(CineAssemblyConfigWindow, ParentWindow);
+						})),
+						EUserInterfaceActionType::None
+					);
 				}
 			}
 		}));
diff --git a/Engine/Plugins/VirtualProduction/Shots/Source/ShotsEditor/Private/UI/CineAssembly/SCineAssemblyConfigWindow.cpp b/Engine/Plugins/VirtualProduction/Shots/Source/ShotsEditor/Private/UI/CineAssembly/SCineAssemblyConfigWindow.cpp
--- a/Engine/Plugins/VirtualProduction/Shots/Source/ShotsEditor/Private/UI/CineAssembly/SCineAssemblyConfigWindow.cpp
+++ b/Engine/Plugins/VirtualProduction/Shots/Source/ShotsEditor/Private/UI/CineAssembly/SCineAssemblyConfigWindow.cpp
@@ -15,17 +15,32 @@
 
 #define LOCTEXT_NAMESPACE "SCineAssemblyConfigWindow"
 
-void SCineAssemblyConfigWindow::Construct(const FArguments& InArgs, UCineAssembly* InCineAssembly)
+void SCineAssemblyConfigWindow::Construct(const FArguments& InArgs, UCineAssembly* InCineAssembly, ECineAssemblyConfigMode InConfigMode)
 {
 	CineAssemblyToConfigure = InCineAssembly;
+	Mode = InConfigMode;
 
-	UWorld* CurrentWorld = GEditor->GetEditorWorldContext().World();
-	CineAssemblyToConfigure->Level = FSoftObjectPath(CurrentWorld);
+	// When creating a new CineAssembly, initialize its Level to the currently open level
+	if (Mode == ECineAssemblyConfigMode::CreateNew)
+	{
+		UWorld* CurrentWorld = GEditor->GetEditorWorldContext().World();
+		CineAssemblyToConfigure->Level = FSoftObjectPath(CurrentWorld);
+	}
 
 	const FVector2D DefaultWindowSize = FVector2D(1400, 750);
 
+	FText WindowTitle;
+	if (Mode == ECineAssemblyConfigMode::CreateNew)
+	{
+		WindowTitle = LOCTEXT("WindowTitleCreateNew", "Create Cine Assembly");
+	}
+	else if (Mode == ECineAssemblyConfigMode::Edit)
+	{
+		WindowTitle = LOCTEXT("WindowTitleEdit", "Edit Cine Assembly");
+	}
+
 	SWindow::Construct(SWindow::FArguments()
-		.Title((LOCTEXT("CineAssemblyCreationWindowTitle", "Create Cine Assembly")))
+		.Title(WindowTitle)
 		.SizingRule(ESizingRule::UserSized)
 		.ClientSize(DefaultWindowSize)
 		.SupportsMaximize(false)
@@ -107,37 +122,42 @@ void SCineAssemblyConfigWindow::MakeDetailsPanel()
 
 TSharedRef<SWidget> SCineAssemblyConfigWindow::MakeButtonsPanel()
 {
-	return SNew(SBorder)
-		.BorderImage(FAppStyle::Get().GetBrush("Brushes.Panel"))
-		.Padding(16.0f)
-		.HAlign(HAlign_Right)
-		[
-			SNew(SHorizontalBox)
-
-			+ SHorizontalBox::Slot()
-				.MinWidth(118.0f)
-				.MaxWidth(118.0f)
-				.VAlign(VAlign_Center)
-				.Padding(0.0f, 0.0f, 8.0f, 0.0f)
-				[
-					SNew(SButton)
-						.Text(LOCTEXT("CreateAssetButton", "Create Asset"))
-						.ButtonStyle(FAppStyle::Get(), "PrimaryButton")
-						.HAlign(HAlign_Center)
-						.OnClicked(this, &SCineAssemblyConfigWindow::OnCreateAssetClicked)
-				]
-
-			+ SHorizontalBox::Slot()
-				.MinWidth(118.0f)
-				.MaxWidth(118.0f)
-				.VAlign(VAlign_Center)
-				[
-					SNew(SButton)
-						.Text(LOCTEXT("CancelButton", "Cancel"))
-						.HAlign(HAlign_Center)
-						.OnClicked(this, &SCineAssemblyConfigWindow::OnCancelClicked)
-				]
-		];
+	if (Mode == ECineAssemblyConfigMode::CreateNew)
+	{
+		return SNew(SBorder)
+			.BorderImage(FAppStyle::Get().GetBrush("Brushes.Panel"))
+			.Padding(16.0f)
+			.HAlign(HAlign_Right)
+			[
+				SNew(SHorizontalBox)
+
+				+ SHorizontalBox::Slot()
+					.MinWidth(118.0f)
+					.MaxWidth(118.0f)
+					.VAlign(VAlign_Center)
+					.Padding(0.0f, 0.0f, 8.0f, 0.0f)
+					[
+						SNew(SButton)
+							.Text(LOCTEXT("CreateAssetButton", "Create Asset"))
+							.ButtonStyle(FAppStyle::Get(), "PrimaryButton")
+							.HAlign(HAlign_Center)
+							.OnClicked(this, &SCineAssemblyConfigWindow::OnCreateAssetClicked)
+					]
+
+				+ SHorizontalBox::Slot()
+					.MinWidth(118.0f)
+					.MaxWidth(118.0f)
+					.VAlign(VAlign_Center)
+					[
+						SNew(SButton)
+							.Text(LOCTEXT("CancelButton", "Cancel"))
+							.HAlign(HAlign_Center)
+							.OnClicked(this, &SCineAssemblyConfigWindow::OnCancelClicked)
+					]
+			];
+	}
+
+	return SNullWidget::NullWidget;
 }
 
 FReply SCineAssemblyConfigWindow::OnCreateAssetClicked()
diff --git a/Engine/Plugins/VirtualProduction/Shots/Source/ShotsEditor/Private/UI/CineAssembly/SCineAssemblyConfigWindow.h b/Engine/Plugins/VirtualProduction/Shots/Source/ShotsEditor/Private/UI/CineAssembly/SCineAssemblyConfigWindow.h
--- a/Engine/Plugins/VirtualProduction/Shots/Source/ShotsEditor/Private/UI/CineAssembly/SCineAssemblyConfigWindow.h
+++ b/Engine/Plugins/VirtualProduction/Shots/Source/ShotsEditor/Private/UI/CineAssembly/SCineAssemblyConfigWindow.h
@@ -8,6 +8,13 @@
 #include "Widgets/Layout/SSplitter.h"
 #include "Widgets/SWindow.h"
 
+/** Mode that configures the UI based on the intended user interactions with the Cine Assembly asset */
+enum class ECineAssemblyConfigMode : uint8
+{
+	CreateNew,
+	Edit
+};
+
 /**
  * A window for configuring the properties of a UCineAssembly asset
  */
@@ -19,7 +26,7 @@ public:
 	SLATE_BEGIN_ARGS(SCineAssemblyConfigWindow) {}
 	SLATE_END_ARGS()
 
-	void Construct(const FArguments& InArgs, UCineAssembly* InCineAssembly);
+	void Construct(const FArguments& InArgs, UCineAssembly* InCineAssembly, ECineAssemblyConfigMode InConfigMode);
 
 	/** Whether a new asset should be created after closing this window */
 	bool ShouldCreateNewAsset() { return bShouldCreateNewAsset; }
@@ -50,4 +57,7 @@ private:
 
 	/** Transient object used only by this UI to configure the properties of the new asset that will get created by the Factory */
 	UCineAssembly* CineAssemblyToConfigure = nullptr;
+
+	/** Mode that configures the UI based on the intended user interactions with the Cine Assembly asset */
+	ECineAssemblyConfigMode Mode = ECineAssemblyConfigMode::CreateNew;
 };
