ShotManagement: Added custom metadata and notes properties to Cine Assembly assets.
Added tab controls to the Cine Assembly config window.
Cine Assembly notes are visible in the tooltip on the asset tile in Content Browser.
#rnx

diff --git a/Engine/Plugins/VirtualProduction/Shots/Source/ShotsEditor/Private/AssetDefinition_CineAssembly.cpp b/Engine/Plugins/VirtualProduction/Shots/Source/ShotsEditor/Private/AssetDefinition_CineAssembly.cpp
--- a/Engine/Plugins/VirtualProduction/Shots/Source/ShotsEditor/Private/AssetDefinition_CineAssembly.cpp
+++ b/Engine/Plugins/VirtualProduction/Shots/Source/ShotsEditor/Private/AssetDefinition_CineAssembly.cpp
@@ -15,6 +15,19 @@ FText UAssetDefinition_CineAssembly::GetAssetDisplayName() const
 	return LOCTEXT("AssetTypeActions_CineAssembly", "Cine Assembly");
 }
 
+FText UAssetDefinition_CineAssembly::GetAssetDescription(const FAssetData& AssetData) const
+{
+	if (const UCineAssembly* const CineAssembly = Cast<UCineAssembly>(AssetData.GetAsset()))
+	{
+		if (!CineAssembly->AssemblyNote.IsEmpty())
+		{
+			return FText::FromString(CineAssembly->AssemblyNote);
+		}
+	}
+
+	return FText::GetEmpty();
+}
+
 TSoftClassPtr<> UAssetDefinition_CineAssembly::GetAssetClass() const
 {
 	return UCineAssembly::StaticClass();
diff --git a/Engine/Plugins/VirtualProduction/Shots/Source/ShotsEditor/Private/AssetDefinition_CineAssembly.h b/Engine/Plugins/VirtualProduction/Shots/Source/ShotsEditor/Private/AssetDefinition_CineAssembly.h
--- a/Engine/Plugins/VirtualProduction/Shots/Source/ShotsEditor/Private/AssetDefinition_CineAssembly.h
+++ b/Engine/Plugins/VirtualProduction/Shots/Source/ShotsEditor/Private/AssetDefinition_CineAssembly.h
@@ -15,6 +15,7 @@ class UAssetDefinition_CineAssembly : public UAssetDefinition_LevelSequence
 protected:
 	// Begin UAssetDefinition Interface
 	virtual FText GetAssetDisplayName() const override;
+	virtual FText GetAssetDescription(const FAssetData& AssetData) const override;
 	virtual TSoftClassPtr<UObject> GetAssetClass() const override;
 	virtual FLinearColor GetAssetColor() const override;
 	virtual TConstArrayView<FAssetCategoryPath> GetAssetCategories() const override;
diff --git a/Engine/Plugins/VirtualProduction/Shots/Source/ShotsEditor/Private/CineAssembly.cpp b/Engine/Plugins/VirtualProduction/Shots/Source/ShotsEditor/Private/CineAssembly.cpp
new file mode 100644
--- /dev/null
+++ b/Engine/Plugins/VirtualProduction/Shots/Source/ShotsEditor/Private/CineAssembly.cpp
@@ -0,0 +1,17 @@
+// Copyright Epic Games, Inc. All Rights Reserved.
+
+#include "CineAssembly.h"
+#include "UObject/AssetRegistryTagsContext.h"
+
+void UCineAssembly::GetAssetRegistryTags(FAssetRegistryTagsContext Context) const
+{
+	Super::GetAssetRegistryTags(Context);
+
+	for (const TPair<FName, FString>& Pair : CustomMetadata)
+	{
+		if (!Pair.Key.IsNone() && !Pair.Value.IsEmpty())
+		{
+			Context.AddTag(FAssetRegistryTag(Pair.Key, Pair.Value, FAssetRegistryTag::ETagType::TT_Alphabetical, FAssetRegistryTag::TD_None));
+		}
+	}
+}
diff --git a/Engine/Plugins/VirtualProduction/Shots/Source/ShotsEditor/Private/CineAssembly.h b/Engine/Plugins/VirtualProduction/Shots/Source/ShotsEditor/Private/CineAssembly.h
--- a/Engine/Plugins/VirtualProduction/Shots/Source/ShotsEditor/Private/CineAssembly.h
+++ b/Engine/Plugins/VirtualProduction/Shots/Source/ShotsEditor/Private/CineAssembly.h
@@ -17,6 +17,19 @@ class UCineAssembly : public ULevelSequence
 public:
 	UCineAssembly() = default;
 
+	//~ Begin UObject Interface
+	virtual void GetAssetRegistryTags(FAssetRegistryTagsContext Context) const override;
+	//~ End UObject Interface
+
+	/** The level to open before opening this asset in Sequencer */
 	UPROPERTY(EditAnywhere, Category = "TargetLevel", meta = (AllowedClasses = "/Script/Engine.World"))
 	FSoftObjectPath Level;
+
+	/** User added metadata key/value pairs, which will be added as additional asset registry tags */
+	UPROPERTY(EditAnywhere, Category = "Metadata")
+	TMap<FName, FString> CustomMetadata;
+
+	/** User-facing notes about this assembly asset */
+	UPROPERTY()
+	FString AssemblyNote;
 };
diff --git a/Engine/Plugins/VirtualProduction/Shots/Source/ShotsEditor/Private/ShotsEditorStyle.cpp b/Engine/Plugins/VirtualProduction/Shots/Source/ShotsEditor/Private/ShotsEditorStyle.cpp
--- a/Engine/Plugins/VirtualProduction/Shots/Source/ShotsEditor/Private/ShotsEditorStyle.cpp
+++ b/Engine/Plugins/VirtualProduction/Shots/Source/ShotsEditor/Private/ShotsEditorStyle.cpp
@@ -34,6 +34,8 @@ FShotsEditorStyle::FShotsEditorStyle()
 	Set("Icons.Folder", new FSlateVectorImageBrush(RootToContentDir(TEXT("Slate/folder-closed.svg")), Icon16x16));
 	Set("Icons.Export", new FSlateVectorImageBrush(RootToContentDir(TEXT("Slate/ExportAll.svg")), Icon16x16));
 	Set("Icons.Initialize", new FSlateVectorImageBrush(RootToContentDir(TEXT("Slate/Initialize_16.svg")), Icon16x16));
+	Set("Icons.Animation", new FSlateVectorImageBrush(RootToContentDir(TEXT("Slate/Animation_16.svg")), Icon16x16));
+	Set("Icons.Notes", new FSlateVectorImageBrush(RootToContentDir(TEXT("Slate/Note.svg")), Icon16x16));
 
 	Set("Badges.RevisionControlWarning", new FSlateVectorImageBrush(RootToCoreContentDir(TEXT("Starship/SourceControl/Status/RevisionControlBadgeWarning.svg")), Icon16x16, FStyleColors::Warning));
 	Set("Badges.RevisionControlConnected", new FSlateVectorImageBrush(RootToCoreContentDir(TEXT("Starship/SourceControl/Status/RevisionControlBadgeConnected.svg")), Icon16x16, FStyleColors::Success));
@@ -49,6 +51,8 @@ FShotsEditorStyle::FShotsEditorStyle()
 	Set("ProductionWizard.PanelBackground", new FSlateRoundedBoxBrush(FStyleColors::Panel, 0.0f, FStyleColors::Recessed, 2.0f));
 	Set("ProductionWizard.RecessedBackground", new FSlateRoundedBoxBrush(FStyleColors::Recessed, 0.0f, FStyleColors::Panel, 2.0f));
 
+	Set("Borders.Background", new FSlateRoundedBoxBrush(FStyleColors::Background, 0.0f, FStyleColors::Panel, 2.0f));
+
 	// Buttons
 	const FButtonStyle RecessedButton = FButtonStyle()
 		.SetNormal(FSlateRoundedBoxBrush(FStyleColors::Recessed, 0.0f, FStyleColors::Recessed, 0.0f))
diff --git a/Engine/Plugins/VirtualProduction/Shots/Source/ShotsEditor/Private/UI/CineAssembly/SCineAssemblyConfigWindow.cpp b/Engine/Plugins/VirtualProduction/Shots/Source/ShotsEditor/Private/UI/CineAssembly/SCineAssemblyConfigWindow.cpp
--- a/Engine/Plugins/VirtualProduction/Shots/Source/ShotsEditor/Private/UI/CineAssembly/SCineAssemblyConfigWindow.cpp
+++ b/Engine/Plugins/VirtualProduction/Shots/Source/ShotsEditor/Private/UI/CineAssembly/SCineAssemblyConfigWindow.cpp
@@ -9,8 +9,11 @@
 #include "ShotsEditorStyle.h"
 #include "Widgets/SBoxPanel.h"
 #include "Widgets/Input/SButton.h"
+#include "Widgets/Input/SSegmentedControl.h"
 #include "Widgets/Layout/SBorder.h"
 #include "Widgets/Layout/SSeparator.h"
+#include "Widgets/Layout/SSplitter.h"
+#include "Widgets/Text/SMultiLineEditableText.h"
 #include "Widgets/Text/STextBlock.h"
 
 #define LOCTEXT_NAMESPACE "SCineAssemblyConfigWindow"
@@ -58,12 +61,16 @@ void SCineAssemblyConfigWindow::Construct(const FArguments& InArgs, UCineAssembl
 								.PhysicalSplitterHandleSize(2.0f)
 
 								+SSplitter::Slot()
-								.Value(0.5f)
-								.Expose(CineTemplatePanelSlot)
+								.Value(0.7f)
+								[
+									MakeCineTemplatePanel()
+								]
 
 								+ SSplitter::Slot()
-								.Value(0.5f)
-								.Expose(DetailsPanelSlot)
+								.Value(0.3f)
+								[
+									MakeInfoPanel()
+								]
 						]
 
 					+ SVerticalBox::Slot()
@@ -81,9 +88,6 @@ void SCineAssemblyConfigWindow::Construct(const FArguments& InArgs, UCineAssembl
 						]
 				]
 		]);
-
-	CineTemplatePanelSlot->AttachWidget(MakeCineTemplatePanel());
-	MakeDetailsPanel();
 }
 
 TSharedRef<SWidget> SCineAssemblyConfigWindow::MakeCineTemplatePanel()
@@ -106,7 +110,72 @@ TSharedRef<SWidget> SCineAssemblyConfigWindow::MakeCineTemplatePanel()
 		];
 }
 
-void SCineAssemblyConfigWindow::MakeDetailsPanel()
+TSharedRef<SWidget> SCineAssemblyConfigWindow::MakeInfoPanel()
+{
+	// TODO: Remove this when all of the tab content has been implemented
+	auto BuildStubWidget = [](const FString& StubText) -> TSharedRef<SWidget>
+		{
+			return SNew(SHorizontalBox)
+				+ SHorizontalBox::Slot()
+				.FillWidth(1.0f)
+				.VAlign(VAlign_Center)
+				.HAlign(HAlign_Center)
+				[
+					SNew(STextBlock).Text(FText::FromString(StubText))
+				];
+		};
+
+	TabSwitcher = SNew(SWidgetSwitcher)
+		+ SWidgetSwitcher::Slot()
+			[ BuildStubWidget("OVERVIEW TAB CONTENT") ]
+
+		+ SWidgetSwitcher::Slot()
+			[ MakeDetailsWidget() ]
+
+		+ SWidgetSwitcher::Slot()
+			[ BuildStubWidget("HIERARCHY TAB CONTENT") ]
+
+		+ SWidgetSwitcher::Slot()
+			[ MakeNotesWidget() ];
+
+	return SNew(SVerticalBox)
+
+		+ SVerticalBox::Slot()
+		.AutoHeight()
+		[
+			SNew(SSegmentedControl<int32>)
+				.Value(0)
+				.OnValueChanged_Lambda([Switcher = TabSwitcher](int32 NewValue)
+				{
+					Switcher->SetActiveWidgetIndex(NewValue);
+				})
+
+				+ SSegmentedControl<int32>::Slot(0)
+				.Text(LOCTEXT("OverviewTab", "Overview"))
+				.Icon(FSlateIcon(FShotsEditorStyle::StyleName, "Icons.Animation").GetIcon())
+
+				+ SSegmentedControl<int32>::Slot(1)
+				.Text(LOCTEXT("DetailsTab", "Details"))
+				.Icon(FSlateIcon(FAppStyle::GetAppStyleSetName(), "Icons.Details").GetIcon())
+
+				+ SSegmentedControl<int32>::Slot(2)
+				.Text(LOCTEXT("HierarchyTab", "Hierarchy"))
+				.Icon(FSlateIcon(FAppStyle::GetAppStyleSetName(), "Icons.FolderClosed").GetIcon())
+
+				+ SSegmentedControl<int32>::Slot(3)
+				.Text(LOCTEXT("NotesTab", "Notes"))
+				.Icon(FSlateIcon(FShotsEditorStyle::StyleName, "Icons.Notes").GetIcon())
+		]
+
+		+ SVerticalBox::Slot()
+		.FillHeight(1.0f)
+		.HAlign(HAlign_Fill)
+		[
+			TabSwitcher.ToSharedRef()
+		];
+}
+
+TSharedRef<SWidget> SCineAssemblyConfigWindow::MakeDetailsWidget()
 {
 	FPropertyEditorModule& PropertyEditor = FModuleManager::Get().LoadModuleChecked<FPropertyEditorModule>(TEXT("PropertyEditor"));
 
@@ -116,9 +185,44 @@ void SCineAssemblyConfigWindow::MakeDetailsPanel()
 
 	TSharedRef<IDetailsView> DetailsView = PropertyEditor.CreateDetailView(DetailsViewArgs);
 	DetailsView->SetObject(CineAssemblyToConfigure, true);
-	DetailsPanelSlot->AttachWidget(DetailsView);
+
+	return DetailsView;
 }
 
+TSharedRef<SWidget> SCineAssemblyConfigWindow::MakeNotesWidget()
+{
+	return SNew(SVerticalBox)
+
+		+ SVerticalBox::Slot()
+		.Padding(16.0f)
+		.AutoHeight()
+		.HAlign(HAlign_Left)
+		[
+			SNew(STextBlock)
+				.Text(LOCTEXT("NoteInstructions", "The following notes will be stored on the shot. This can also be edited later."))
+				.AutoWrapText(true)
+		]
+
+		+ SVerticalBox::Slot()
+		[
+			SNew(SBorder)
+				.BorderImage(FShotsEditorStyle::Get().GetBrush("Borders.Background"))
+				.Padding(16.0f)
+				[
+					SNew(SMultiLineEditableText)
+						.HintText(LOCTEXT("NoteHintText", "What initial notes do you have about this shot?\nYou can always add more later."))
+						.Text_Lambda([this]()
+						{ 
+							return FText::FromString(CineAssemblyToConfigure->AssemblyNote); 
+						})
+						.OnTextCommitted_Lambda([this](const FText& InText, ETextCommit::Type InCommitType)
+						{
+							CineAssemblyToConfigure->Modify();
+							CineAssemblyToConfigure->AssemblyNote = InText.ToString();
+						})
+				]
+		];
+}
 
 TSharedRef<SWidget> SCineAssemblyConfigWindow::MakeButtonsPanel()
 {
diff --git a/Engine/Plugins/VirtualProduction/Shots/Source/ShotsEditor/Private/UI/CineAssembly/SCineAssemblyConfigWindow.h b/Engine/Plugins/VirtualProduction/Shots/Source/ShotsEditor/Private/UI/CineAssembly/SCineAssemblyConfigWindow.h
--- a/Engine/Plugins/VirtualProduction/Shots/Source/ShotsEditor/Private/UI/CineAssembly/SCineAssemblyConfigWindow.h
+++ b/Engine/Plugins/VirtualProduction/Shots/Source/ShotsEditor/Private/UI/CineAssembly/SCineAssemblyConfigWindow.h
@@ -5,7 +5,7 @@
 #include "Widgets/SCompoundWidget.h"
 
 #include "CineAssembly.h"
-#include "Widgets/Layout/SSplitter.h"
+#include "Widgets/Layout/SWidgetSwitcher.h"
 #include "Widgets/SWindow.h"
 
 /** Mode that configures the UI based on the intended user interactions with the Cine Assembly asset */
@@ -32,15 +32,21 @@ public:
 	bool ShouldCreateNewAsset() { return bShouldCreateNewAsset; }
 
 private:
-	/** Creates the menu on the left side of the production wizard */
+	/** Creates the panel that displays the available Cine Assembly templates */
 	TSharedRef<SWidget> MakeCineTemplatePanel();
 
-	/** Creates the menu on the left side of the production wizard */
-	void MakeDetailsPanel();
+	/** Creates the panel that holds the various tabs with properties of the Cine Assembly */
+	TSharedRef<SWidget> MakeInfoPanel();
 
-	/** Creates the buttons on the bottom of the production wizard */
+	/** Creates the buttons on the bottom of the window */
 	TSharedRef<SWidget> MakeButtonsPanel();
 
+	/** Creates the widget to display for the Details tab */
+	TSharedRef<SWidget> MakeDetailsWidget();
+
+	/** Creates the widget to display for the Notes tab */
+	TSharedRef<SWidget> MakeNotesWidget();
+
 	/** Cycles backwards through the available panels */
 	FReply OnCreateAssetClicked();
 
@@ -48,9 +54,8 @@ private:
 	FReply OnCancelClicked();
 
 private:
-	/** Slots for the content panels */
-	SSplitter::FSlot* CineTemplatePanelSlot = nullptr;
-	SSplitter::FSlot* DetailsPanelSlot = nullptr;
+	/** Switcher that controls which tab widget is currently visible */
+	TSharedPtr<SWidgetSwitcher> TabSwitcher;
 
 	/** Whether a new asset should be created after closing this window */
 	bool bShouldCreateNewAsset = false;
