ShotManagement: Add CineAssemblySchema asset type, definition, factory, and config window.
The config window currently just displays a simple details view of the schema asset. This includes a texture picker for a thumbnail image which is used as the thumbnail brush for the widget after it is created.
Users cannot rename schema assets from the Content Browser, but they can rename them by entering a new value for the SchemaName property. This is in keeping with the design choice that the asset name should always match the SchemaName property.

diff --git a/Engine/Plugins/VirtualProduction/Shots/Source/ShotsEditor/Private/AssetDefinition_CineAssemblySchema.cpp b/Engine/Plugins/VirtualProduction/Shots/Source/ShotsEditor/Private/AssetDefinition_CineAssemblySchema.cpp
new file mode 100644
--- /dev/null
+++ b/Engine/Plugins/VirtualProduction/Shots/Source/ShotsEditor/Private/AssetDefinition_CineAssemblySchema.cpp
@@ -0,0 +1,76 @@
+// Copyright Epic Games, Inc. All Rights Reserved.
+
+#include "AssetDefinition_CineAssemblySchema.h"
+
+#include "CineAssemblySchema.h"
+#include "Framework/Application/SlateApplication.h"
+#include "Interfaces/IMainFrameModule.h"
+#include "Modules/ModuleManager.h"
+#include "UI/CineAssembly/SCineAssemblySchemaWindow.h"
+
+#define LOCTEXT_NAMESPACE "AssetTypeActions"
+
+TSoftClassPtr<> UAssetDefinition_CineAssemblySchema::GetAssetClass() const
+{
+	return UCineAssemblySchema::StaticClass();
+}
+
+FText UAssetDefinition_CineAssemblySchema::GetAssetDisplayName() const
+{
+	return LOCTEXT("AssetTypeActions_CineAssemblySchema", "Cine Assembly Schema");
+}
+
+FLinearColor UAssetDefinition_CineAssemblySchema::GetAssetColor() const
+{
+	return FColor(176, 58, 104);
+}
+
+TConstArrayView<FAssetCategoryPath> UAssetDefinition_CineAssemblySchema::GetAssetCategories() const
+{
+	static const FAssetCategoryPath Categories[] = { EAssetCategoryPaths::Cinematics };
+	return Categories;
+}
+
+const FSlateBrush* UAssetDefinition_CineAssemblySchema::GetThumbnailBrush(const FAssetData& InAssetData, const FName InClassName) const
+{
+	if (UCineAssemblySchema* CineAssemblySchema = Cast<UCineAssemblySchema>(InAssetData.GetAsset()))
+	{
+		return CineAssemblySchema->GetThumbnailBrush();
+	}
+
+	return nullptr;
+}
+
+FAssetSupportResponse UAssetDefinition_CineAssemblySchema::CanRename(const FAssetData& InAsset) const
+{
+	if (UCineAssemblySchema* CineAssemblySchema = Cast<UCineAssemblySchema>(InAsset.GetAsset()))
+	{
+		if (!CineAssemblySchema->SupportsRename())
+		{
+			return FAssetSupportResponse::NotSupported();
+		}
+	}
+
+	return FAssetSupportResponse::Supported();
+}
+
+EAssetCommandResult UAssetDefinition_CineAssemblySchema::OpenAssets(const FAssetOpenArgs& OpenArgs) const
+{
+	// We only support opening one asset at a time
+	if (OpenArgs.Assets.Num() == 1)
+	{
+		const FAssetData& SchemaAssetData = OpenArgs.Assets[0];
+		if (UCineAssemblySchema* CineAssemblySchema = Cast<UCineAssemblySchema>(SchemaAssetData.GetAsset()))
+		{
+			IMainFrameModule& MainFrame = FModuleManager::LoadModuleChecked<IMainFrameModule>("MainFrame");
+			TSharedPtr<SWindow>	ParentWindow = MainFrame.GetParentWindow();
+
+			TSharedRef<SCineAssemblySchemaWindow> CineAssemblySchemaWindow = SNew(SCineAssemblySchemaWindow, Cast<UCineAssemblySchema>(CineAssemblySchema), ESchemaConfigMode::Edit);
+			FSlateApplication::Get().AddModalWindow(CineAssemblySchemaWindow, ParentWindow);
+		}
+	}
+
+	return EAssetCommandResult::Handled;
+}
+
+#undef LOCTEXT_NAMESPACE
diff --git a/Engine/Plugins/VirtualProduction/Shots/Source/ShotsEditor/Private/AssetDefinition_CineAssemblySchema.h b/Engine/Plugins/VirtualProduction/Shots/Source/ShotsEditor/Private/AssetDefinition_CineAssemblySchema.h
new file mode 100644
--- /dev/null
+++ b/Engine/Plugins/VirtualProduction/Shots/Source/ShotsEditor/Private/AssetDefinition_CineAssemblySchema.h
@@ -0,0 +1,25 @@
+// Copyright Epic Games, Inc. All Rights Reserved.
+
+#pragma once
+
+#include "AssetDefinition.h"
+
+#include "AssetDefinition_CineAssemblySchema.generated.h"
+
+/** Asset definition for a UCineAssemblySchema asset */
+UCLASS()
+class UAssetDefinition_CineAssemblySchema : public UAssetDefinition
+{
+	GENERATED_BODY()
+
+protected:
+	// Begin UAssetDefinition Interface
+	virtual TSoftClassPtr<UObject> GetAssetClass() const override;
+	virtual FText GetAssetDisplayName() const override;
+	virtual FLinearColor GetAssetColor() const override;
+	virtual TConstArrayView<FAssetCategoryPath> GetAssetCategories() const override;
+	virtual const FSlateBrush* GetThumbnailBrush(const FAssetData& InAssetData, const FName InClassName) const override;
+	virtual EAssetCommandResult OpenAssets(const FAssetOpenArgs& OpenArgs) const override;
+	virtual FAssetSupportResponse CanRename(const FAssetData& InAsset) const override;
+	// End UAssetDefinition Interface
+};
diff --git a/Engine/Plugins/VirtualProduction/Shots/Source/ShotsEditor/Private/CineAssemblySchema.cpp b/Engine/Plugins/VirtualProduction/Shots/Source/ShotsEditor/Private/CineAssemblySchema.cpp
new file mode 100644
--- /dev/null
+++ b/Engine/Plugins/VirtualProduction/Shots/Source/ShotsEditor/Private/CineAssemblySchema.cpp
@@ -0,0 +1,79 @@
+// Copyright Epic Games, Inc. All Rights Reserved.
+
+#include "CineAssemblySchema.h"
+
+#include "AssetToolsModule.h"
+#include "Modules/ModuleManager.h"
+
+void UCineAssemblySchema::PostLoad()
+{
+	Super::PostLoad();
+
+	UpdateThumbnailBrush();
+}
+
+void UCineAssemblySchema::PostEditChangeProperty(struct FPropertyChangedEvent& PropertyChangedEvent)
+{
+	const FName PropertyName = PropertyChangedEvent.GetPropertyName();
+
+	if (PropertyName == GET_MEMBER_NAME_CHECKED(UCineAssemblySchema, ThumbnailImage))
+	{
+		// We mark the new thumbnail as pending until the next time that the brush is requested from an external caller.
+		// This is because Slate holds on to the brush pointer, and assumes (in many places) that it will remain valid.
+		// So if we reset the brush pointer or reset its resource object now, we run the risk of crashing within Slate
+		// somewhere when it tries to access the stale brush.
+		bThumbnailPendingReset = true;
+	}
+}
+
+void UCineAssemblySchema::RenameAsset()
+{
+	// The default behavior for schema assets is to not allow renaming from Content Browser.
+	// However, this function relies on renaming being supported, so we temporarily enable to do the programmatic rename.
+	bSupportsRename = true;
+
+	IAssetTools& AssetTools = FModuleManager::GetModuleChecked<FAssetToolsModule>("AssetTools").Get();
+
+	const FSoftObjectPath OldObjectPath = FSoftObjectPath(this);
+
+	const FTopLevelAssetPath NewAssetPath(*GetPathName(), *SchemaName);
+	const FSoftObjectPath NewObjectPath(NewAssetPath);
+
+	FAssetRenameData RenameData = FAssetRenameData(OldObjectPath, NewObjectPath);
+	AssetTools.RenameAssets({ RenameData });
+
+	bSupportsRename = false;
+}
+
+bool UCineAssemblySchema::SupportsRename()
+{
+	return bSupportsRename;
+}
+
+void UCineAssemblySchema::UpdateThumbnailBrush()
+{
+	if (ThumbnailImage)
+	{
+		if (!ThumbnailBrush.IsValid())
+		{
+			ThumbnailBrush = MakeShareable(new FSlateBrush());
+		}
+
+		ThumbnailBrush->SetResourceObject(ThumbnailImage);
+		ThumbnailBrush->ImageSize = FVector2D(64.0f, 64.0f);
+	}
+	else
+	{
+		ThumbnailBrush.Reset();
+	}
+}
+
+FSlateBrush* UCineAssemblySchema::GetThumbnailBrush()
+{
+	if (bThumbnailPendingReset)
+	{
+		UpdateThumbnailBrush();
+	}
+
+	return ThumbnailBrush.Get();
+}
diff --git a/Engine/Plugins/VirtualProduction/Shots/Source/ShotsEditor/Private/CineAssemblySchema.h b/Engine/Plugins/VirtualProduction/Shots/Source/ShotsEditor/Private/CineAssemblySchema.h
new file mode 100644
--- /dev/null
+++ b/Engine/Plugins/VirtualProduction/Shots/Source/ShotsEditor/Private/CineAssemblySchema.h
@@ -0,0 +1,61 @@
+// Copyright Epic Games, Inc. All Rights Reserved.
+
+#pragma once
+
+#include "Engine/Texture2D.h"
+#include "Styling/SlateBrush.h"
+
+#include "CineAssemblySchema.generated.h"
+
+/**
+ * A template object for building different Cine Assembly types
+ */
+UCLASS()
+class UCineAssemblySchema : public UObject
+{
+	GENERATED_BODY()
+
+public:
+	UCineAssemblySchema() = default;
+
+	//~ Begin UObject overrides
+	virtual void PostLoad() override;
+	virtual void PostEditChangeProperty(struct FPropertyChangedEvent& PropertyChangedEvent) override;
+	//~ End UObject overrides
+
+	/** Renames the underlying schema asset */
+	void RenameAsset();
+
+	/** Whether or not this schema can be renamed */
+	bool SupportsRename();
+
+	/** Returns the thumbnail brush used by this schema */
+	FSlateBrush* GetThumbnailBrush();
+
+private:
+	/** Updates the thumbnail brush resource */
+	void UpdateThumbnailBrush();
+
+public:
+	/** The schema name, which will be used by assemblies made from this schema as their "assembly type" */
+	UPROPERTY(EditAnywhere, Category = "Default")
+	FString SchemaName;
+
+	/** A user-facing text description of this schema */
+	UPROPERTY(EditAnywhere, Category = "Default")
+	FString Description;
+
+	/** The thumbnail image to use for this schema and assemblies built from this schema */
+	UPROPERTY(EditAnywhere, Category = "Default")
+	TObjectPtr<UTexture2D> ThumbnailImage;
+
+private:
+	/** Slate brush used to actually render the thumbnail image */
+	TSharedPtr<FSlateBrush> ThumbnailBrush;
+
+	/** Whether the thumbnail brush needs to be updated */
+	bool bThumbnailPendingReset = false;
+
+	/** Whether the schema asset supports renaming */
+	bool bSupportsRename = false;
+};
diff --git a/Engine/Plugins/VirtualProduction/Shots/Source/ShotsEditor/Private/CineAssemblySchemaFactory.cpp b/Engine/Plugins/VirtualProduction/Shots/Source/ShotsEditor/Private/CineAssemblySchemaFactory.cpp
new file mode 100644
--- /dev/null
+++ b/Engine/Plugins/VirtualProduction/Shots/Source/ShotsEditor/Private/CineAssemblySchemaFactory.cpp
@@ -0,0 +1,75 @@
+// Copyright Epic Games, Inc. All Rights Reserved.
+
+#include "CineAssemblySchemaFactory.h"
+
+#include "AssetToolsModule.h"
+#include "Framework/Application/SlateApplication.h"
+#include "Interfaces/IMainFrameModule.h"
+#include "Modules/ModuleManager.h"
+#include "UI/CineAssembly/SCineAssemblySchemaWindow.h"
+#include "UObject/Package.h"
+
+UCineAssemblySchemaFactory::UCineAssemblySchemaFactory()
+{
+	SupportedClass = UCineAssemblySchema::StaticClass();
+	bCreateNew = true;
+	bEditAfterNew = true;
+
+	SchemaToConfigure = nullptr;
+}
+
+bool UCineAssemblySchemaFactory::CanCreateNew() const
+{
+	return true;
+}
+
+UObject* UCineAssemblySchemaFactory::FactoryCreateNew(UClass* Class, UObject* InParent, FName Name, EObjectFlags Flags, UObject* Context, FFeedbackContext* Warn)
+{
+	if (SchemaToConfigure)
+	{
+		// If the user configured a name for the schema, use that as the name for the new asset. Otherwise, use the input default name
+		const FString NewSchemaName = SchemaToConfigure->SchemaName.IsEmpty() ? Name.ToString() : SchemaToConfigure->SchemaName;
+
+		// Update the name, outer package, and flags of the configuration asset to no longer be transient
+		SchemaToConfigure->Rename(*NewSchemaName, InParent);
+		SchemaToConfigure->SetFlags(Flags);
+		SchemaToConfigure->ClearFlags(RF_Transient);
+
+		// Unregister the temporary default name that we registered earlier for this schema instance
+		IAssetTools& AssetTools = FModuleManager::GetModuleChecked<FAssetToolsModule>("AssetTools").Get();
+		AssetTools.UnregisterDefaultAssetNameForClass(UCineAssemblySchema::StaticClass());
+	}
+
+	return SchemaToConfigure;
+}
+
+bool UCineAssemblySchemaFactory::ConfigureProperties()
+{
+	// Create a new transient CineAssembly to configure in the UI.
+	// If the configuration is successful, this will turn into the persistent object created by this factory.
+	SchemaToConfigure = NewObject<UCineAssemblySchema>(GetTransientPackage(), NAME_None, RF_Transient);
+
+	// Spawn a new window to configure the properties of the asset being created
+	IMainFrameModule& MainFrame = FModuleManager::LoadModuleChecked<IMainFrameModule>("MainFrame");
+	TSharedPtr<SWindow>	ParentWindow = MainFrame.GetParentWindow();
+
+	TSharedRef<SCineAssemblySchemaWindow> CineAssemblySchemaWindow = SNew(SCineAssemblySchemaWindow, SchemaToConfigure, ESchemaConfigMode::CreateNew);
+	FSlateApplication::Get().AddModalWindow(CineAssemblySchemaWindow, ParentWindow);
+
+	if (CineAssemblySchemaWindow->ShouldCreateNewAsset() == false)
+	{
+		return false;
+	}
+
+	// One of the properties configured by the user in the schema window is the schema name.
+	// We want the actual asset name to match the value of this property, and we don't want the user to have to retype it.
+	// This workaround allows us to temporarily set a default name for the schema asset. 
+	// After FactoryCreateNew() is finished, we will unregister this default name override. 
+	if (!SchemaToConfigure->SchemaName.IsEmpty())
+	{
+		IAssetTools& AssetTools = FModuleManager::GetModuleChecked<FAssetToolsModule>("AssetTools").Get();
+		AssetTools.RegisterDefaultAssetNameForClass(UCineAssemblySchema::StaticClass(), SchemaToConfigure->SchemaName);
+	}
+
+	return true;
+}
diff --git a/Engine/Plugins/VirtualProduction/Shots/Source/ShotsEditor/Private/CineAssemblySchemaFactory.h b/Engine/Plugins/VirtualProduction/Shots/Source/ShotsEditor/Private/CineAssemblySchemaFactory.h
new file mode 100644
--- /dev/null
+++ b/Engine/Plugins/VirtualProduction/Shots/Source/ShotsEditor/Private/CineAssemblySchemaFactory.h
@@ -0,0 +1,31 @@
+// Copyright Epic Games, Inc. All Rights Reserved.
+
+#pragma once
+
+#include "Factories/Factory.h"
+#include "CineAssemblySchema.h"
+
+#include "CineAssemblySchemaFactory.generated.h"
+
+/**
+ * Factory class used to create new UCineAssemblySchema objects
+ * Before creating a new assembly schema, the factory will spawn a window to configure the properties of the asset that is being created
+ */
+UCLASS(hidecategories = Object)
+class UCineAssemblySchemaFactory : public UFactory
+{
+	GENERATED_BODY()
+
+protected:
+	UCineAssemblySchemaFactory();
+
+	// Begin UFactory Interface
+	virtual UObject* FactoryCreateNew(UClass* Class, UObject* InParent, FName Name, EObjectFlags Flags, UObject* Context, FFeedbackContext* Warn) override;
+	virtual bool CanCreateNew() const override;
+	virtual bool ConfigureProperties() override;
+	// End UFactory Interface
+
+private:
+	/** The schema object that will be created by this factory and configured by the user */
+	UCineAssemblySchema* SchemaToConfigure;
+};
diff --git a/Engine/Plugins/VirtualProduction/Shots/Source/ShotsEditor/Private/ShotsEditorStyle.cpp b/Engine/Plugins/VirtualProduction/Shots/Source/ShotsEditor/Private/ShotsEditorStyle.cpp
--- a/Engine/Plugins/VirtualProduction/Shots/Source/ShotsEditor/Private/ShotsEditorStyle.cpp
+++ b/Engine/Plugins/VirtualProduction/Shots/Source/ShotsEditor/Private/ShotsEditorStyle.cpp
@@ -4,6 +4,7 @@
 
 #include "Interfaces/IPluginManager.h"
 #include "Styling/CoreStyle.h"
+#include "Styling/SegmentedControlStyle.h"
 #include "Styling/SlateStyleMacros.h"
 #include "Styling/SlateTypes.h"
 #include "Styling/StyleColors.h"
@@ -36,6 +37,7 @@ FShotsEditorStyle::FShotsEditorStyle()
 	Set("Icons.Initialize", new FSlateVectorImageBrush(RootToContentDir(TEXT("Slate/Initialize_16.svg")), Icon16x16));
 	Set("Icons.Animation", new FSlateVectorImageBrush(RootToContentDir(TEXT("Slate/Animation_16.svg")), Icon16x16));
 	Set("Icons.Notes", new FSlateVectorImageBrush(RootToContentDir(TEXT("Slate/Note.svg")), Icon16x16));
+	Set("Icons.Shot", new FSlateVectorImageBrush(RootToContentDir(TEXT("Slate/Shot_16.svg")), Icon16x16));
 
 	Set("Badges.RevisionControlWarning", new FSlateVectorImageBrush(RootToCoreContentDir(TEXT("Starship/SourceControl/Status/RevisionControlBadgeWarning.svg")), Icon16x16, FStyleColors::Warning));
 	Set("Badges.RevisionControlConnected", new FSlateVectorImageBrush(RootToCoreContentDir(TEXT("Starship/SourceControl/Status/RevisionControlBadgeConnected.svg")), Icon16x16, FStyleColors::Success));
@@ -70,6 +72,31 @@ FShotsEditorStyle::FShotsEditorStyle()
 
 	Set("ProductionWizard.PanelButton", PanelButton);
 
+	// Segmented Controls
+	FCheckBoxStyle PrimarySegmentedBoxStyle = FCheckBoxStyle()
+		.SetCheckBoxType(ESlateCheckBoxType::ToggleButton)
+		.SetUncheckedImage(FSlateNoResource())
+		.SetUncheckedHoveredImage(FSlateRoundedBoxBrush(FStyleColors::Hover, 4.0f))
+		.SetUncheckedPressedImage(FSlateRoundedBoxBrush(FStyleColors::Header, 4.0f))
+		.SetCheckedImage(FSlateRoundedBoxBrush(FStyleColors::Primary, 4.0f))
+		.SetCheckedHoveredImage(FSlateRoundedBoxBrush(FStyleColors::PrimaryHover, 4.0f))
+		.SetCheckedPressedImage(FSlateRoundedBoxBrush(FStyleColors::PrimaryPress, 4.0f))
+		.SetForegroundColor(FStyleColors::Foreground)
+		.SetHoveredForegroundColor(FStyleColors::ForegroundHover)
+		.SetPressedForegroundColor(FStyleColors::ForegroundHover)
+		.SetCheckedForegroundColor(FStyleColors::ForegroundHover)
+		.SetCheckedHoveredForegroundColor(FStyleColors::ForegroundHover)
+		.SetCheckedPressedForegroundColor(FStyleColors::ForegroundHover)
+		.SetPadding(FMargin(8, 3));
+
+	Set("PrimarySegmentedControl", FSegmentedControlStyle()
+		.SetControlStyle(PrimarySegmentedBoxStyle)
+		.SetFirstControlStyle(PrimarySegmentedBoxStyle)
+		.SetLastControlStyle(PrimarySegmentedBoxStyle)
+		.SetBackgroundBrush(FSlateRoundedBoxBrush(FStyleColors::Recessed, 4.0f))
+		.SetUniformPadding(FMargin(1))
+	);
+
 	FSlateStyleRegistry::RegisterSlateStyle(*this);
 }
 
diff --git a/Engine/Plugins/VirtualProduction/Shots/Source/ShotsEditor/Private/UI/CineAssembly/SCineAssemblySchemaWindow.cpp b/Engine/Plugins/VirtualProduction/Shots/Source/ShotsEditor/Private/UI/CineAssembly/SCineAssemblySchemaWindow.cpp
new file mode 100644
--- /dev/null
+++ b/Engine/Plugins/VirtualProduction/Shots/Source/ShotsEditor/Private/UI/CineAssembly/SCineAssemblySchemaWindow.cpp
@@ -0,0 +1,248 @@
+// Copyright Epic Games, Inc. All Rights Reserved.
+
+#include "SCineAssemblySchemaWindow.h"
+
+#include "IStructureDetailsView.h"
+#include "Modules/ModuleManager.h"
+#include "PropertyEditorModule.h"
+#include "ShotsEditorStyle.h"
+#include "Widgets/Input/SButton.h"
+#include "Widgets/Input/SSegmentedControl.h"
+#include "Widgets/Layout/SBorder.h"
+#include "Widgets/Layout/SSeparator.h"
+#include "Widgets/Layout/SSplitter.h"
+#include "Widgets/SBoxPanel.h"
+#include "Widgets/Text/STextBlock.h"
+
+#define LOCTEXT_NAMESPACE "SCineAssemblySchemaWindow"
+
+void SCineAssemblySchemaWindow::Construct(const FArguments& InArgs, UCineAssemblySchema* InSchema, ESchemaConfigMode InConfigMode)
+{
+	SchemaToConfigure = InSchema;
+	Mode = InConfigMode;
+
+	const FVector2D DefaultWindowSize = FVector2D(1400, 750);
+
+	FText WindowTitle;
+	if (Mode == ESchemaConfigMode::CreateNew)
+	{
+		WindowTitle = LOCTEXT("WindowTitleCreateNew", "Create Assembly Schema");
+	}
+	else if (Mode == ESchemaConfigMode::Edit)
+	{
+		WindowTitle = LOCTEXT("WindowTitleEdit", "Edit Assembly Schema");
+	}
+
+	GetOnWindowClosedEvent().AddSP(this, &SCineAssemblySchemaWindow::OnSchemaWindowClosed);
+
+	SWindow::Construct(SWindow::FArguments()
+		.Title(WindowTitle)
+		.SizingRule(ESizingRule::UserSized)
+		.ClientSize(DefaultWindowSize)
+		.SupportsMaximize(false)
+		.SupportsMinimize(false)
+		[
+			SNew(SBorder)
+				.BorderImage(FShotsEditorStyle::Get().GetBrush("ProductionWizard.PanelBackground"))
+				[
+					SNew(SVerticalBox)
+
+					+ SVerticalBox::Slot()
+						[
+							SNew(SSplitter)
+								.Orientation(Orient_Horizontal)
+								.PhysicalSplitterHandleSize(2.0f)
+
+							+ SSplitter::Slot()
+								.Value(0.2f)
+								[
+									MakeMenuPanel()
+								]
+
+							+ SSplitter::Slot()
+								.Value(0.8f)
+								[
+									MakeContentPanel()
+								]
+						]
+
+					+ SVerticalBox::Slot()
+						.AutoHeight()
+						[
+							SNew(SSeparator)
+								.Orientation(Orient_Horizontal)
+								.Thickness(2.0f)
+						]
+
+					+ SVerticalBox::Slot()
+						.AutoHeight()
+						[
+							MakeButtonsPanel()
+						]
+				]
+		]);
+}
+
+TSharedRef<SWidget> SCineAssemblySchemaWindow::MakeMenuPanel()
+{
+	return SNew(SBorder)
+		.BorderImage(FShotsEditorStyle::Get().GetBrush("ProductionWizard.RecessedBackground"))
+		.Padding(8.0f)
+		.VAlign(VAlign_Top)
+		[
+			SNew(SSegmentedControl<int32>)
+				.Style(&FShotsEditorStyle::Get().GetWidgetStyle<FSegmentedControlStyle>("PrimarySegmentedControl"))
+				.MaxSegmentsPerLine(1)
+				.Value_Lambda([this]()
+				{ 
+					return TabSwitcher->GetActiveWidgetIndex(); 
+				})
+				.OnValueChanged_Lambda([this](int32 NewValue)
+				{
+					TabSwitcher->SetActiveWidgetIndex(NewValue);
+				})
+
+				+ SSegmentedControl<int32>::Slot(0)
+				.Text(LOCTEXT("AssetTab", "Asset"))
+				.Icon(FSlateIcon(FShotsEditorStyle::StyleName, "Icons.Shot").GetIcon())
+				.HAlign(HAlign_Left)
+
+				+ SSegmentedControl<int32>::Slot(1)
+				.Text(LOCTEXT("TimelineTab", "Timeline Template"))
+				.Icon(FSlateIcon(FShotsEditorStyle::StyleName, "Icons.Sequencer").GetIcon())
+				.HAlign(HAlign_Left)
+
+				+ SSegmentedControl<int32>::Slot(2)
+				.Text(LOCTEXT("HierarchyTab", "Hierarchy"))
+				.Icon(FSlateIcon(FAppStyle::GetAppStyleSetName(), "Icons.FolderClosed").GetIcon())
+				.HAlign(HAlign_Left)
+		];
+}
+
+TSharedRef<SWidget> SCineAssemblySchemaWindow::MakeContentPanel()
+{
+	// TODO: Remove this when all of the tab content has been implemented
+	auto BuildStubWidget = [](const FString& StubText) -> TSharedRef<SWidget>
+		{
+			return SNew(SHorizontalBox)
+
+			+ SHorizontalBox::Slot()
+				.FillWidth(1.0f)
+				.VAlign(VAlign_Center)
+				.HAlign(HAlign_Center)
+				[
+					SNew(STextBlock).Text(FText::FromString(StubText))
+				];
+		};
+
+	TabSwitcher = SNew(SWidgetSwitcher)
+		+ SWidgetSwitcher::Slot()
+		[
+			MakeAssetTabContent()
+		]
+
+		+ SWidgetSwitcher::Slot()
+		[
+			BuildStubWidget("TIMELINE TAB CONTENT")
+		]
+
+		+ SWidgetSwitcher::Slot()
+		[
+			BuildStubWidget("HIERARCHY TAB CONTENT")
+		];
+
+	return TabSwitcher.ToSharedRef();
+}
+
+TSharedRef<SWidget> SCineAssemblySchemaWindow::MakeButtonsPanel()
+{
+	if (Mode == ESchemaConfigMode::CreateNew)
+	{
+		return SNew(SBorder)
+			.BorderImage(FAppStyle::Get().GetBrush("Brushes.Panel"))
+			.Padding(16.0f)
+			.HAlign(HAlign_Right)
+			[
+				SNew(SHorizontalBox)
+
+				+ SHorizontalBox::Slot()
+					.MinWidth(118.0f)
+					.MaxWidth(118.0f)
+					.VAlign(VAlign_Center)
+					.Padding(0.0f, 0.0f, 8.0f, 0.0f)
+					[
+						SNew(SButton)
+							.Text(LOCTEXT("CreateAssetButton", "Create Asset"))
+							.ButtonStyle(FAppStyle::Get(), "PrimaryButton")
+							.HAlign(HAlign_Center)
+							.OnClicked(this, &SCineAssemblySchemaWindow::OnCreateAssetClicked)
+					]
+
+				+ SHorizontalBox::Slot()
+					.MinWidth(118.0f)
+					.MaxWidth(118.0f)
+					.VAlign(VAlign_Center)
+					[
+						SNew(SButton)
+							.Text(LOCTEXT("CancelButton", "Cancel"))
+							.HAlign(HAlign_Center)
+							.OnClicked(this, &SCineAssemblySchemaWindow::OnCancelClicked)
+					]
+			];
+	}
+
+	return SNullWidget::NullWidget;
+}
+
+TSharedRef<SWidget> SCineAssemblySchemaWindow::MakeAssetTabContent()
+{
+	FPropertyEditorModule& PropertyEditor = FModuleManager::Get().LoadModuleChecked<FPropertyEditorModule>(TEXT("PropertyEditor"));
+
+	FDetailsViewArgs DetailsViewArgs;
+	DetailsViewArgs.bAllowSearch = false;
+	DetailsViewArgs.bHideSelectionTip = true;
+
+	TSharedRef<IDetailsView> DetailsView = PropertyEditor.CreateDetailView(DetailsViewArgs);
+	DetailsView->SetObject(SchemaToConfigure, true);
+
+	DetailsView->OnFinishedChangingProperties().AddSP(this, &SCineAssemblySchemaWindow::OnSchemaPropertiesChanged);
+
+	return DetailsView;
+}
+
+void SCineAssemblySchemaWindow::OnSchemaPropertiesChanged(const FPropertyChangedEvent& PropertyChangedEvent)
+{
+	const FName PropertyName = PropertyChangedEvent.GetPropertyName();
+
+	if (PropertyName == GET_MEMBER_NAME_CHECKED(UCineAssemblySchema, SchemaName))
+	{
+		if (Mode == ESchemaConfigMode::Edit)
+		{
+			bPendingRename = true;
+		}
+	}
+}
+
+FReply SCineAssemblySchemaWindow::OnCreateAssetClicked()
+{
+	bShouldCreateNewAsset = true;
+	RequestDestroyWindow();
+	return FReply::Handled();
+}
+
+FReply SCineAssemblySchemaWindow::OnCancelClicked()
+{
+	bShouldCreateNewAsset = false;
+	RequestDestroyWindow();
+	return FReply::Handled();
+}
+
+void SCineAssemblySchemaWindow::OnSchemaWindowClosed(const TSharedRef<SWindow>& InWindow)
+{
+	if (bPendingRename)
+	{
+		SchemaToConfigure->RenameAsset();
+	}
+}
+
+#undef LOCTEXT_NAMESPACE
diff --git a/Engine/Plugins/VirtualProduction/Shots/Source/ShotsEditor/Private/UI/CineAssembly/SCineAssemblySchemaWindow.h b/Engine/Plugins/VirtualProduction/Shots/Source/ShotsEditor/Private/UI/CineAssembly/SCineAssemblySchemaWindow.h
new file mode 100644
--- /dev/null
+++ b/Engine/Plugins/VirtualProduction/Shots/Source/ShotsEditor/Private/UI/CineAssembly/SCineAssemblySchemaWindow.h
@@ -0,0 +1,73 @@
+// Copyright Epic Games, Inc. All Rights Reserved.
+
+#pragma once
+
+#include "Widgets/SWindow.h"
+
+#include "CineAssemblySchema.h"
+#include "Widgets/Layout/SWidgetSwitcher.h"
+
+/** Mode that configures the UI based on the intended user interactions with the Cine Assembly Schema asset */
+enum class ESchemaConfigMode : uint8
+{
+	CreateNew,
+	Edit
+};
+
+/**
+ * A window for configuring the properties of a UCineAssemblySchema
+ */
+class SCineAssemblySchemaWindow : public SWindow
+{
+public:
+	SCineAssemblySchemaWindow() = default;
+
+	SLATE_BEGIN_ARGS(SCineAssemblySchemaWindow) {}
+	SLATE_END_ARGS()
+
+	void Construct(const FArguments& InArgs, UCineAssemblySchema* InSchema, ESchemaConfigMode InConfigMode);
+
+	/** Whether a new asset should be created after closing this window */
+	bool ShouldCreateNewAsset() { return bShouldCreateNewAsset; }
+
+private:
+	/** Creates the panel that displays the tab menu */
+	TSharedRef<SWidget> MakeMenuPanel();
+
+	/** Creates the panel that displays the content for each tab */
+	TSharedRef<SWidget> MakeContentPanel();
+
+	/** Creates the buttons on the bottom of the window */
+	TSharedRef<SWidget> MakeButtonsPanel();
+
+	/** Creates the content for the Assets tab */
+	TSharedRef<SWidget> MakeAssetTabContent();
+
+	/** Closes the window and indicates that a new asset should be created by the asset factory */
+	FReply OnCreateAssetClicked();
+
+	/** Closes the window and indicates that no assets should be created by the asset factory */
+	FReply OnCancelClicked();
+
+	/** Callback when one of the properties of the schema being configured changes */
+	void OnSchemaPropertiesChanged(const FPropertyChangedEvent& PropertyChangedEvent);
+
+	/** Callback when this window is about to be closed */
+	void OnSchemaWindowClosed(const TSharedRef<SWindow>& InWindow);
+
+private:
+	/** Switcher that controls which tab widget is currently visible */
+	TSharedPtr<SWidgetSwitcher> TabSwitcher;
+
+	/** Whether a new asset should be created after closing this window */
+	bool bShouldCreateNewAsset = false;
+
+	/** Transient object used only by this UI to configure the properties of the new asset that will get created by the factory */
+	UCineAssemblySchema* SchemaToConfigure = nullptr;
+
+	/** Mode that configures the UI based on the intended user interactions with the Cine Assembly schema asset */
+	ESchemaConfigMode Mode = ESchemaConfigMode::CreateNew;
+
+	/** Whether the schema asset being configured needs to be renamed */
+	bool bPendingRename = false;
+};
