ShotManagement: Fix Production Settings API to no longer return a pointer to the current active production. Instead, the API can return the GUID of the active production, or a const copy of it.
Changing the active production's properties has intentional side effects, so we want to ensure that callers cannot directly modify it.
It was also possible for those pointers to become stale if the Production Settings array was ever reallocated.
#rnx

diff --git a/Engine/Plugins/VirtualProduction/Shots/Source/ShotsEditor/Private/ProductionFunctionLibrary.cpp b/Engine/Plugins/VirtualProduction/Shots/Source/ShotsEditor/Private/ProductionFunctionLibrary.cpp
--- a/Engine/Plugins/VirtualProduction/Shots/Source/ShotsEditor/Private/ProductionFunctionLibrary.cpp
+++ b/Engine/Plugins/VirtualProduction/Shots/Source/ShotsEditor/Private/ProductionFunctionLibrary.cpp
@@ -2,10 +2,10 @@
 
 #include "ProductionFunctionLibrary.h"
 
-TArray<FCinematicProduction> UProductionFunctionLibrary::GetAllProductions()
+const TArray<FCinematicProduction> UProductionFunctionLibrary::GetAllProductions()
 {
 	const UProductionSettings* const ProductionSettings = GetDefault<UProductionSettings>();
-	return ProductionSettings->Productions;
+	return ProductionSettings->GetProductions();
 }
 
 void UProductionFunctionLibrary::SetActiveProduction(FCinematicProduction Production)
diff --git a/Engine/Plugins/VirtualProduction/Shots/Source/ShotsEditor/Private/ProductionFunctionLibrary.h b/Engine/Plugins/VirtualProduction/Shots/Source/ShotsEditor/Private/ProductionFunctionLibrary.h
--- a/Engine/Plugins/VirtualProduction/Shots/Source/ShotsEditor/Private/ProductionFunctionLibrary.h
+++ b/Engine/Plugins/VirtualProduction/Shots/Source/ShotsEditor/Private/ProductionFunctionLibrary.h
@@ -19,7 +19,7 @@ public:
 
 	/** Returns an array of all available Cinematic Productions */
 	UFUNCTION(BlueprintPure, Category = "Shot Management")
-	static TArray<FCinematicProduction> GetAllProductions();
+	static const TArray<FCinematicProduction> GetAllProductions();
 
 	/** 
 	 * Sets the input Production as the current Active Production 
diff --git a/Engine/Plugins/VirtualProduction/Shots/Source/ShotsEditor/Private/ProductionSettings.cpp b/Engine/Plugins/VirtualProduction/Shots/Source/ShotsEditor/Private/ProductionSettings.cpp
--- a/Engine/Plugins/VirtualProduction/Shots/Source/ShotsEditor/Private/ProductionSettings.cpp
+++ b/Engine/Plugins/VirtualProduction/Shots/Source/ShotsEditor/Private/ProductionSettings.cpp
@@ -1,6 +1,9 @@
 // Copyright Epic Games, Inc. All Rights Reserved.
 #include "ProductionSettings.h"
 
+#include "Algo/Contains.h"
+#include "Algo/Find.h"
+#include "Algo/IndexOf.h"
 #include "AssetToolsModule.h"
 #include "Engine/Engine.h"
 #include "HAL/FileManager.h"
@@ -80,10 +83,11 @@ void UProductionSettings::PostInitProperties()
 	}
 
 	// Initialize the active production based on the serialized active production name
-	ActiveProduction = Productions.FindByPredicate([ProductionName = ActiveProductionName](const FCinematicProduction& Other) { return Other.ProductionName == ProductionName; });
+	const FCinematicProduction* ActiveProduction = Algo::FindBy(Productions, ActiveProductionName, &FCinematicProduction::ProductionName);
+	ActiveProductionID = ActiveProduction ? ActiveProduction->ProductionID : FGuid();
 
 	// Apply the initial overrides to existing project settings based on the current active production
-	if (ActiveProduction)
+	if (ActiveProductionID.IsValid())
 	{
 		ApplyProjectOverrides();
 	}
@@ -118,7 +122,7 @@ void UProductionSettings::ApplyProjectOverrides()
 	UMovieSceneToolsProjectSettings* MovieSceneToolsSettings = GetMutableDefault<UMovieSceneToolsProjectSettings>();
  	FProperty* DefaultStartTimeProperty = MovieSceneToolsSettings->GetClass()->FindPropertyByName(GET_MEMBER_NAME_CHECKED(UMovieSceneToolsProjectSettings, DefaultStartTime));
 
-	if (ActiveProduction)
+	if (ActiveProductionID.IsValid())
 	{
 		DefaultDisplayRateProperty->SetPropertyFlags(EPropertyFlags::CPF_EditConst);
 		DefaultStartTimeProperty->SetPropertyFlags(EPropertyFlags::CPF_EditConst);
@@ -149,7 +153,7 @@ void UProductionSettings::OverrideDefaultDisplayRate()
 	// If there is no active production, the level sequence project setting is reset to its default config value
 	ULevelSequenceProjectSettings* LevelSequenceSettings = GetMutableDefault<ULevelSequenceProjectSettings>();
 
-	if (ActiveProduction)
+	if (const FCinematicProduction* ActiveProduction = GetActiveProductionPtr())
 	{
 		const FText FrameRateText = ActiveProduction->DefaultDisplayRate.ToPrettyText();
 		LevelSequenceSettings->DefaultDisplayRate = FrameRateText.ToString();
@@ -172,7 +176,7 @@ void UProductionSettings::OverrideDefaultStartTime()
 	// If there is no active production, the movie scene tools project setting is reset to its default config value
 	UMovieSceneToolsProjectSettings* MovieSceneToolsSettings = GetMutableDefault<UMovieSceneToolsProjectSettings>();
 
-	if (ActiveProduction)
+	if (const FCinematicProduction* ActiveProduction = GetActiveProductionPtr())
 	{
 		MovieSceneToolsSettings->DefaultStartTime = ActiveProduction->DefaultDisplayRate.AsSeconds(FFrameNumber(ActiveProduction->DefaultStartFrame));
 	}
@@ -184,7 +188,7 @@ void UProductionSettings::OverrideDefaultStartTime()
 
 void UProductionSettings::OverrideSubsequenceHierarchicalBias()
 {
-	if (ActiveProduction)
+	if (const FCinematicProduction* ActiveProduction = GetActiveProductionPtr())
 	{
 		using namespace UE::Shots::ProductionSettings::Private;
 		const int32 HierarchicalBiasValue = (ActiveProduction->SubsequencePriority == ESubsequencePriority::TopDown) ? TopDownHBiasValue : BottomUpHBiasValue;
@@ -205,7 +209,7 @@ void UProductionSettings::OverrideSubsequenceHierarchicalBias()
 
 void UProductionSettings::OverrideDefaultAssetNames()
 {
-	if (ActiveProduction)
+	if (const FCinematicProduction* ActiveProduction = GetActiveProductionPtr())
 	{
 		IAssetTools& AssetTools = FModuleManager::GetModuleChecked<FAssetToolsModule>("AssetTools").Get();
 
@@ -240,18 +244,16 @@ void UProductionSettings::PostEditChangeProperty(struct FPropertyChangedEvent& P
 			// Make sure the production name is unique
 			const int32 ArrayIndex = PropertyChangedEvent.GetArrayIndex(PropertyChangedEvent.MemberProperty->GetFName().ToString());
 			Productions[ArrayIndex].ProductionName = GetUniqueProductionName();
-
-			// Reset the active production pointer, in case the array was reallocated
-			ActiveProduction = Productions.FindByPredicate([this](const FCinematicProduction& Other) { return Other.ProductionName == ActiveProductionName; });
 		}
 		else
 		{
 			// It is possible that a production was just deleted. If it was the active production, then that needs to be updated.
-			// This attempts to reacquire the pointer to the active production, and if it no longer exists, it will be reset to nullptr.
-			ActiveProduction = Productions.FindByPredicate([this](const FCinematicProduction& Other) { return Other.ProductionName == ActiveProductionName; });
+			const FCinematicProduction* ActiveProduction = Algo::FindBy(Productions, ActiveProductionName, &FCinematicProduction::ProductionName);
+			ActiveProductionID = ActiveProduction ? ActiveProduction->ProductionID : FGuid();
+
 			SetActiveProductionName();
 
-			if (!ActiveProduction)
+			if (!ActiveProductionID.IsValid())
 			{
 				OverrideDefaultDisplayRate();
 			}
@@ -291,9 +293,6 @@ void UProductionSettings::AddProduction(const FCinematicProduction& ProductionTo
 {
 	Productions.Add(ProductionToAdd);
 
-	// Reset the active production pointer, in case the array was reallocated
-	ActiveProduction = Productions.FindByPredicate([this](const FCinematicProduction& Other) { return Other.ProductionName == ActiveProductionName; });
-
 	TryUpdateDefaultConfigFile();
 
 	ProductionListChangedDelegate.Broadcast();
@@ -301,8 +300,7 @@ void UProductionSettings::AddProduction(const FCinematicProduction& ProductionTo
 
 void UProductionSettings::DuplicateProduction(FGuid ProductionID)
 {
-	const FCinematicProduction* ProductionToCopy = Productions.FindByPredicate([ProductionID](const FCinematicProduction& Other) { return Other.ProductionID == ProductionID; });
-	if (ProductionToCopy)
+	if (const FCinematicProduction* ProductionToCopy = Algo::FindBy(Productions, ProductionID, &FCinematicProduction::ProductionID))
 	{
 		FCinematicProduction DuplicateProduction = *ProductionToCopy;
 
@@ -312,9 +310,6 @@ void UProductionSettings::DuplicateProduction(FGuid ProductionID)
 
 		Productions.Add(DuplicateProduction);
 
-		// Reset the active production pointer, in case the array was reallocated
-		ActiveProduction = Productions.FindByPredicate([this](const FCinematicProduction& Other) { return Other.ProductionName == ActiveProductionName; });
-
 		TryUpdateDefaultConfigFile();
 
 		ProductionListChangedDelegate.Broadcast();
@@ -328,7 +323,7 @@ void UProductionSettings::DeleteProduction(FGuid ProductionID)
 		return;
 	}
 
-	int32 IndexToDelete = Productions.IndexOfByPredicate([ProductionID](const FCinematicProduction& Other) { return Other.ProductionID == ProductionID; });
+	const int32 IndexToDelete = Algo::IndexOfBy(Productions, ProductionID, &FCinematicProduction::ProductionID);
 
 	if (IndexToDelete != INDEX_NONE)
 	{
@@ -337,9 +332,6 @@ void UProductionSettings::DeleteProduction(FGuid ProductionID)
 
 		Productions.RemoveAt(IndexToDelete);
 
-		// Reset the active production pointer
-		ActiveProduction = Productions.FindByPredicate([this](const FCinematicProduction& Other) { return Other.ProductionName == ActiveProductionName; });
-
 		if (bIsActiveProduction)
 		{
 			SetActiveProduction(ProductionID);
@@ -353,8 +345,7 @@ void UProductionSettings::DeleteProduction(FGuid ProductionID)
 
 void UProductionSettings::RenameProduction(FGuid ProductionID, FString NewName)
 {
-	FCinematicProduction* Production = Productions.FindByPredicate([ProductionID](const FCinematicProduction& Other) { return Other.ProductionID == ProductionID; });
-	if (Production)
+	if (FCinematicProduction* Production = Algo::FindBy(Productions, ProductionID, &FCinematicProduction::ProductionID))
 	{
 		Production->ProductionName = NewName;
 		SetActiveProductionName();
@@ -362,30 +353,48 @@ void UProductionSettings::RenameProduction(FGuid ProductionID, FString NewName)
 	}
 }
 
-const FCinematicProduction* UProductionSettings::GetProduction(FGuid ProductionID) const
+const TArray<FCinematicProduction> UProductionSettings::GetProductions() const
 {
-	if (!ProductionID.IsValid())
+	return Productions;
+}
+
+TOptional<const FCinematicProduction> UProductionSettings::GetProduction(FGuid ProductionID) const
+{
+	if (const FCinematicProduction* Production = Algo::FindBy(Productions, ProductionID, &FCinematicProduction::ProductionID))
 	{
-		return nullptr;
+		return *Production;
 	}
+	return TOptional<const FCinematicProduction>();
+}
 
-	return Productions.FindByPredicate([ProductionID](const FCinematicProduction& Other) { return Other.ProductionID == ProductionID; });
+TOptional<const FCinematicProduction> UProductionSettings::GetActiveProduction() const
+{
+	return GetProduction(ActiveProductionID);
+}
+
+FGuid UProductionSettings::GetActiveProductionID() const
+{
+	return ActiveProductionID;
 }
 
-FCinematicProduction* UProductionSettings::GetActiveProduction() const
+const FCinematicProduction* UProductionSettings::GetActiveProductionPtr() const
 {
-	return ActiveProduction;
+	if (ActiveProductionID.IsValid())
+	{
+		return Algo::FindBy(Productions, ActiveProductionID, &FCinematicProduction::ProductionID);
+	}
+	return nullptr;
 }
 
 void UProductionSettings::SetActiveProduction(FGuid ProductionID)
 {
 	// If the active production is currently None, before setting it and applying the project overrides, cache some project defaults to restore later if the active production is set to None again
-	if (!ActiveProduction)
+	if (!ActiveProductionID.IsValid())
 	{
 		CacheProjectDefaults();
 	}
 
-	if (ActiveProduction)
+	if (const FCinematicProduction* ActiveProduction = GetActiveProductionPtr())
 	{
 		// Unregister any default asset names that the previously active production had set
 		// Then, restore the corresponding project default asset names
@@ -402,7 +411,16 @@ void UProductionSettings::SetActiveProduction(FGuid ProductionID)
 		}
 	}
 
-	ActiveProduction = Productions.FindByPredicate([ProductionID](const FCinematicProduction& Other) { return Other.ProductionID == ProductionID; });
+	// Update the active production ID to the input ID, if it matches a production that exists in the production list
+	if (Algo::ContainsBy(Productions, ProductionID, &FCinematicProduction::ProductionID))
+	{
+		ActiveProductionID = ProductionID;
+	}
+	else
+	{
+		ActiveProductionID = FGuid();
+	}
+
 	SetActiveProductionName();
 
 	// Apply the overrides based on the new active production (even if it is None, because some properties will be reset to default)
@@ -413,16 +431,16 @@ void UProductionSettings::SetActiveProduction(FGuid ProductionID)
 
 bool UProductionSettings::IsActiveProduction(FGuid ProductionID) const
 {
-	if (ActiveProduction && ProductionID.IsValid())
+	if (ActiveProductionID.IsValid() && ProductionID.IsValid())
 	{
-		return (ProductionID == ActiveProduction->ProductionID);
+		return (ProductionID == ActiveProductionID);
 	}
 	return false;
 }
 
 FFrameRate UProductionSettings::GetActiveDisplayRate() const
 {
-	if (ActiveProduction)
+	if (const FCinematicProduction* ActiveProduction = GetActiveProductionPtr())
 	{
 		return ActiveProduction->DefaultDisplayRate;
 	}
@@ -438,7 +456,7 @@ FFrameRate UProductionSettings::GetActiveDisplayRate() const
 
 int32 UProductionSettings::GetActiveStartFrame() const
 {
-	if (ActiveProduction)
+	if (const FCinematicProduction* ActiveProduction = GetActiveProductionPtr())
 	{
 		return ActiveProduction->DefaultStartFrame;
 	}
@@ -455,7 +473,7 @@ ESubsequencePriority UProductionSettings::GetActiveSubsequencePriority() const
 {
 	using namespace UE::Shots::ProductionSettings::Private;
 
-	if (ActiveProduction)
+	if (const FCinematicProduction* ActiveProduction = GetActiveProductionPtr())
 	{
 		return ActiveProduction->SubsequencePriority;
 	}
@@ -472,8 +490,7 @@ ESubsequencePriority UProductionSettings::GetActiveSubsequencePriority() const
 
 void UProductionSettings::SetDisplayRate(FGuid ProductionID, FFrameRate DisplayRate)
 {
-	FCinematicProduction* Production = Productions.FindByPredicate([ProductionID](const FCinematicProduction& Other) { return Other.ProductionID == ProductionID; });
-	if (Production)
+	if (FCinematicProduction* Production = Algo::FindBy(Productions, ProductionID, &FCinematicProduction::ProductionID))
 	{
 		Production->DefaultDisplayRate = DisplayRate;
 
@@ -489,8 +506,7 @@ void UProductionSettings::SetDisplayRate(FGuid ProductionID, FFrameRate DisplayR
 
 void UProductionSettings::SetStartFrame(FGuid ProductionID, int32 StartFrame)
 {
-	FCinematicProduction* Production = Productions.FindByPredicate([ProductionID](const FCinematicProduction& Other) { return Other.ProductionID == ProductionID; });
-	if (Production)
+	if (FCinematicProduction* Production = Algo::FindBy(Productions, ProductionID, &FCinematicProduction::ProductionID))
 	{
 		Production->DefaultStartFrame = StartFrame;
 
@@ -505,8 +521,7 @@ void UProductionSettings::SetStartFrame(FGuid ProductionID, int32 StartFrame)
 
 void UProductionSettings::SetSubsequencePriority(FGuid ProductionID, ESubsequencePriority Priority)
 {
-	FCinematicProduction* Production = Productions.FindByPredicate([ProductionID](const FCinematicProduction& Other) { return Other.ProductionID == ProductionID; });
-	if (Production)
+	if (FCinematicProduction* Production = Algo::FindBy(Productions, ProductionID, &FCinematicProduction::ProductionID))
 	{
 		Production->SubsequencePriority = Priority;
 
@@ -521,8 +536,7 @@ void UProductionSettings::SetSubsequencePriority(FGuid ProductionID, ESubsequenc
 
 void UProductionSettings::AddNamespaceToDenyList(FGuid ProductionID, const FString& Namespace)
 {
-	FCinematicProduction* Production = Productions.FindByPredicate([ProductionID](const FCinematicProduction& Other) { return Other.ProductionID == ProductionID; });
-	if (Production)
+	if (FCinematicProduction* Production = Algo::FindBy(Productions, ProductionID, &FCinematicProduction::ProductionID))
 	{
 		Production->NamingTokenNamespaceDenyList.Add(Namespace);
 		TryUpdateDefaultConfigFile();
@@ -531,8 +545,7 @@ void UProductionSettings::AddNamespaceToDenyList(FGuid ProductionID, const FStri
 
 void UProductionSettings::RemoveNamespaceFromDenyList(FGuid ProductionID, const FString& Namespace)
 {
-	FCinematicProduction* Production = Productions.FindByPredicate([ProductionID](const FCinematicProduction& Other) { return Other.ProductionID == ProductionID; });
-	if (Production)
+	if (FCinematicProduction* Production = Algo::FindBy(Productions, ProductionID, &FCinematicProduction::ProductionID))
 	{
 		Production->NamingTokenNamespaceDenyList.Remove(Namespace);
 		TryUpdateDefaultConfigFile();
@@ -541,7 +554,7 @@ void UProductionSettings::RemoveNamespaceFromDenyList(FGuid ProductionID, const
 
 void UProductionSettings::FilterNamingTokenNamespaces(TSet<FString>& Namespaces)
 {
-	if (ActiveProduction)
+	if (const FCinematicProduction* ActiveProduction = GetActiveProductionPtr())
 	{
 		for (const FString& Namespace : ActiveProduction->NamingTokenNamespaceDenyList)
 		{
@@ -557,8 +570,7 @@ void UProductionSettings::AddAssetNaming(FGuid ProductionID, const UClass* Asset
 		return;
 	}
 
-	FCinematicProduction* Production = Productions.FindByPredicate([ProductionID](const FCinematicProduction& Other) { return Other.ProductionID == ProductionID; });
-	if (Production)
+	if (FCinematicProduction* Production = Algo::FindBy(Productions, ProductionID, &FCinematicProduction::ProductionID))
 	{
 		Production->DefaultAssetNames.Add(AssetClass, DefaultName);
 
@@ -588,8 +600,7 @@ void UProductionSettings::RemoveAssetNaming(FGuid ProductionID, const UClass* As
 		return;
 	}
 
-	FCinematicProduction* Production = Productions.FindByPredicate([ProductionID](const FCinematicProduction& Other) { return Other.ProductionID == ProductionID; });
-	if (Production)
+	if (FCinematicProduction* Production = Algo::FindBy(Productions, ProductionID, &FCinematicProduction::ProductionID))
 	{
 		Production->DefaultAssetNames.Remove(AssetClass);
 
@@ -613,8 +624,7 @@ void UProductionSettings::RemoveAssetNaming(FGuid ProductionID, const UClass* As
 
 void UProductionSettings::AddTemplateFolder(FGuid ProductionID, const FString& Path, bool bCreateIfMissing)
 {
-	FCinematicProduction* Production = Productions.FindByPredicate([ProductionID](const FCinematicProduction& Other) { return Other.ProductionID == ProductionID; });
-	if (Production)
+	if (FCinematicProduction* Production = Algo::FindBy(Productions, ProductionID, &FCinematicProduction::ProductionID))
 	{
 		FFolderTemplate NewTemplate;
 		NewTemplate.InternalPath = Path;
@@ -627,10 +637,9 @@ void UProductionSettings::AddTemplateFolder(FGuid ProductionID, const FString& P
 
 void UProductionSettings::RemoveTemplateFolder(FGuid ProductionID, const FString& Path)
 {
-	FCinematicProduction* Production = Productions.FindByPredicate([ProductionID](const FCinematicProduction& Other) { return Other.ProductionID == ProductionID; });
-	if (Production)
+	if (FCinematicProduction* Production = Algo::FindBy(Productions, ProductionID, &FCinematicProduction::ProductionID))
 	{
-		int32 IndexToDelete = Production->TemplateFolders.IndexOfByPredicate([Path](const FFolderTemplate& Other) { return Other.InternalPath == Path; });
+		const int32 IndexToDelete = Algo::IndexOfBy(Production->TemplateFolders, Path, &FFolderTemplate::InternalPath);
 		if (Production->TemplateFolders.IsValidIndex(IndexToDelete))
 		{
 			Production->TemplateFolders.RemoveAt(IndexToDelete);
@@ -642,8 +651,7 @@ void UProductionSettings::RemoveTemplateFolder(FGuid ProductionID, const FString
 
 void UProductionSettings::SetTemplateFolderHierarchy(FGuid ProductionID, const TArray<FFolderTemplate>& TemplateHierarchy)
 {
-	FCinematicProduction* Production = Productions.FindByPredicate([ProductionID](const FCinematicProduction& Other) { return Other.ProductionID == ProductionID; });
-	if (Production)
+	if (FCinematicProduction* Production = Algo::FindBy(Productions, ProductionID, &FCinematicProduction::ProductionID))
 	{
 		Production->TemplateFolders = TemplateHierarchy;
 
@@ -653,31 +661,38 @@ void UProductionSettings::SetTemplateFolderHierarchy(FGuid ProductionID, const T
 
 void UProductionSettings::SetActiveDisplayRate(FFrameRate DisplayRate)
 {
-	if (ActiveProduction)
+	if (ActiveProductionID.IsValid())
 	{
-		SetDisplayRate(ActiveProduction->ProductionID, DisplayRate);
+		SetDisplayRate(ActiveProductionID, DisplayRate);
 	}
 }
 
 void UProductionSettings::SetActiveStartFrame(int32 StartFrame)
 {
-	if (ActiveProduction)
+	if (ActiveProductionID.IsValid())
 	{
-		SetStartFrame(ActiveProduction->ProductionID, StartFrame);
+		SetStartFrame(ActiveProductionID, StartFrame);
 	}
 }
 
 void UProductionSettings::SetActiveSubsequencePriority(ESubsequencePriority Priority)
 {
-	if (ActiveProduction)
+	if (ActiveProductionID.IsValid())
 	{
-		SetSubsequencePriority(ActiveProduction->ProductionID, Priority);
+		SetSubsequencePriority(ActiveProductionID, Priority);
 	}
 }
 
 void UProductionSettings::SetActiveProductionName()
 {
-	ActiveProductionName = ActiveProduction ? ActiveProduction->ProductionName : FString(TEXT(""));
+	if (const FCinematicProduction* ActiveProduction = GetActiveProductionPtr())
+	{
+		ActiveProductionName = ActiveProduction->ProductionName;
+	}
+	else
+	{
+		ActiveProductionName = TEXT("");
+	}
 
 	const FString ConfigSection = UProductionSettings::StaticClass()->GetPathName();
 	const FString PropertyName = GET_MEMBER_NAME_CHECKED(UProductionSettings, ActiveProductionName).ToString();
@@ -719,7 +734,8 @@ FString UProductionSettings::GetUniqueProductionName(const FString& BaseName) co
 	}
 
 	FString WorkingName = TrimmedBaseName;
-	while (Productions.ContainsByPredicate([WorkingName](const FCinematicProduction& Other) { return Other.ProductionName == WorkingName; }))
+
+	while (Algo::ContainsBy(Productions, WorkingName, &FCinematicProduction::ProductionName))
 	{
 		WorkingName = FString::Printf(TEXT("%s%d"), *TrimmedBaseName, IntSuffix);
 		IntSuffix++;
diff --git a/Engine/Plugins/VirtualProduction/Shots/Source/ShotsEditor/Private/ProductionSettings.h b/Engine/Plugins/VirtualProduction/Shots/Source/ShotsEditor/Private/ProductionSettings.h
--- a/Engine/Plugins/VirtualProduction/Shots/Source/ShotsEditor/Private/ProductionSettings.h
+++ b/Engine/Plugins/VirtualProduction/Shots/Source/ShotsEditor/Private/ProductionSettings.h
@@ -101,6 +101,15 @@ public:
 	virtual void PostEditChangeProperty(struct FPropertyChangedEvent& PropertyChangedEvent) override;
 	//~ End UObject overrides
 
+	/** The customization is allowed access to the private properties */
+	friend class FProductionSettingsCustomization;
+
+	/** Returns a copy of the production list */
+	const TArray<FCinematicProduction> GetProductions() const;
+
+	/** Returns a copy of the production matching the input production ID (if it exists) */
+	TOptional<const FCinematicProduction> GetProduction(FGuid ProductionID) const;
+
 	/** Adds a new empty production to the list */
 	void AddProduction();
 
@@ -116,11 +125,11 @@ public:
 	/** Renames the production matching the input ID */
 	void RenameProduction(FGuid ProductionID, FString NewName);
 
-	/** Returns a pointer to the production matching the input ID */
-	const FCinematicProduction* GetProduction(FGuid ProductionID) const;
+	/** Returns a copy of the active production (if there is one) */
+	TOptional<const FCinematicProduction> GetActiveProduction() const;
 
-	/** Returns a pointer to the active production */
-	FCinematicProduction* GetActiveProduction() const;
+	/** Returns the unique ID of the active production */
+	FGuid GetActiveProductionID() const;
 
 	/** Sets the active production based on the input production ID */
 	void SetActiveProduction(FGuid ProductionID);
@@ -211,18 +220,20 @@ private:
 	/** Cache the default project settings that are overridden by the active production, used to reset when there is no active production */
 	void CacheProjectDefaults();
 
-public:
+	/** Returns a pointer to the active production (if it exists) */
+	const FCinematicProduction* GetActiveProductionPtr() const;
+
+private:
 	/** Name of the active production */
-	UPROPERTY(VisibleAnywhere, Category = "Default", meta=(DisplayName="Active Production"))
+	UPROPERTY(VisibleAnywhere, Category = "Default", meta = (DisplayName = "Active Production"))
 	FString ActiveProductionName;
 
 	/** List of available productions in this project */
 	UPROPERTY(config, EditAnywhere, Category = "Default")
 	TArray<FCinematicProduction> Productions;
 
-private:
-	/** Pointer to the active production (in the Productions array) */
-	FCinematicProduction* ActiveProduction;
+	/** ID of the active production (in the Productions array) */
+	FGuid ActiveProductionID;
 
 	/** Cached default project settings that are overridden by the active production, used to reset when there is no active production */
 	FString ProjectDefaultDisplayRate;
diff --git a/Engine/Plugins/VirtualProduction/Shots/Source/ShotsEditor/Private/ProductionSettingsCustomization.cpp b/Engine/Plugins/VirtualProduction/Shots/Source/ShotsEditor/Private/ProductionSettingsCustomization.cpp
--- a/Engine/Plugins/VirtualProduction/Shots/Source/ShotsEditor/Private/ProductionSettingsCustomization.cpp
+++ b/Engine/Plugins/VirtualProduction/Shots/Source/ShotsEditor/Private/ProductionSettingsCustomization.cpp
@@ -74,7 +74,8 @@ TSharedRef<SWidget> FProductionSettingsCustomization::BuildProductionNameMenu()
 	const UProductionSettings* const ProductionSettings = GetDefault<UProductionSettings>();
 
 	// Add a menu option with the production name for each production available in this project
-	for (const FCinematicProduction& Production : ProductionSettings->Productions)
+	const TArray<FCinematicProduction> Productions = ProductionSettings->GetProductions();
+	for (const FCinematicProduction& Production : Productions)
 	{
 		MenuBuilder.AddMenuEntry(
 			FText::FromString(Production.ProductionName),
@@ -91,12 +92,16 @@ TSharedRef<SWidget> FProductionSettingsCustomization::BuildProductionNameMenu()
 
 FText FProductionSettingsCustomization::GetActiveProductionName() const
 {
- 	const UProductionSettings* const ProductionSettings = GetDefault<UProductionSettings>();
-	if (ProductionSettings->ActiveProductionName.IsEmpty())
+	const UProductionSettings* const ProductionSettings = GetDefault<UProductionSettings>();
+
+	TOptional<const FCinematicProduction> ActiveProduction = ProductionSettings->GetActiveProduction();
+	const FString ActiveProductionName = ActiveProduction.IsSet() ? ActiveProduction.GetValue().ProductionName : FString();
+
+	if (ActiveProductionName.IsEmpty())
 	{
 		return FText::FromName(NAME_None);
 	}
-	return FText::FromString(ProductionSettings->ActiveProductionName);
+	return FText::FromString(ActiveProductionName);
 }
 
 void FProductionSettingsCustomization::SetActiveProduction(FGuid ProductionID)
diff --git a/Engine/Plugins/VirtualProduction/Shots/Source/ShotsEditor/Private/UI/SActiveProductionCombo.cpp b/Engine/Plugins/VirtualProduction/Shots/Source/ShotsEditor/Private/UI/SActiveProductionCombo.cpp
--- a/Engine/Plugins/VirtualProduction/Shots/Source/ShotsEditor/Private/UI/SActiveProductionCombo.cpp
+++ b/Engine/Plugins/VirtualProduction/Shots/Source/ShotsEditor/Private/UI/SActiveProductionCombo.cpp
@@ -15,11 +15,15 @@ void SActiveProductionCombo::Construct(const FArguments& InArgs)
 	static auto GetActiveProductionName = []() -> FText
 		{
 			const UProductionSettings* const ProductionSettings = GetDefault<UProductionSettings>();
-			if (ProductionSettings->ActiveProductionName.IsEmpty())
+
+			TOptional<const FCinematicProduction> ActiveProduction = ProductionSettings->GetActiveProduction();
+			const FString ActiveProductionName = ActiveProduction.IsSet() ? ActiveProduction.GetValue().ProductionName : FString();
+
+			if (ActiveProductionName.IsEmpty())
 			{
 				return LOCTEXT("NoneProductionName", "None");
 			}
-			return FText::FromString(ProductionSettings->ActiveProductionName);
+			return FText::FromString(ActiveProductionName);
 		};
 
 	static auto OnSetActiveProduction = [](FGuid ProductionID)
@@ -43,7 +47,8 @@ void SActiveProductionCombo::Construct(const FArguments& InArgs)
 			);
 
 			const UProductionSettings* const ProductionSettings = GetDefault<UProductionSettings>();
-			for (const FCinematicProduction& Production : ProductionSettings->Productions)
+			const TArray<FCinematicProduction> Productions = ProductionSettings->GetProductions();
+			for (const FCinematicProduction& Production : Productions)
 			{
 				MenuBuilder.AddMenuEntry(
 					FText::FromString(Production.ProductionName),
diff --git a/Engine/Plugins/VirtualProduction/Shots/Source/ShotsEditor/Private/UI/SAssetNamingPanel.cpp b/Engine/Plugins/VirtualProduction/Shots/Source/ShotsEditor/Private/UI/SAssetNamingPanel.cpp
--- a/Engine/Plugins/VirtualProduction/Shots/Source/ShotsEditor/Private/UI/SAssetNamingPanel.cpp
+++ b/Engine/Plugins/VirtualProduction/Shots/Source/ShotsEditor/Private/UI/SAssetNamingPanel.cpp
@@ -70,7 +70,8 @@ void SAssetNamingPanel::Construct(const FArguments& InArgs)
 	auto IsActiveProductionValid = []() -> bool
 		{
 			const UProductionSettings* const ProductionSettings = GetDefault<UProductionSettings>();
-			return (ProductionSettings->GetActiveProduction() != nullptr);
+			TOptional<const FCinematicProduction> ActiveProduction = ProductionSettings->GetActiveProduction();
+			return ActiveProduction.IsSet();
 		};
 
 	auto OnCreateNewNaming = [this]() -> FReply
@@ -221,9 +222,10 @@ void SAssetNamingPanel::UpdateAssetNamingList()
 	AssetNamingListItems.Empty();
 
 	const UProductionSettings* const ProductionSettings = GetDefault<UProductionSettings>();
-	if (const FCinematicProduction* ActiveProduction = ProductionSettings->GetActiveProduction())
+	TOptional<const FCinematicProduction> ActiveProduction = ProductionSettings->GetActiveProduction();
+	if (ActiveProduction.IsSet())
 	{
-		for (const TPair<TObjectPtr<const UClass>, FString>& AssetNaming : ActiveProduction->DefaultAssetNames)
+		for (const TPair<TObjectPtr<const UClass>, FString>& AssetNaming : ActiveProduction.GetValue().DefaultAssetNames)
 		{
 			AssetNamingListItems.Add(MakeShared<FAssetNamingRowData>(AssetNaming.Key, AssetNaming.Value));
 		}
@@ -326,10 +328,7 @@ void SAssetNamingRow::SetAssetClass(const UClass* SelectedClass)
 
 	// Before updating the asset class for this row, remove the old class from the active production settings list
 	UProductionSettings* ProductionSettings = GetMutableDefault<UProductionSettings>();
-	if (FCinematicProduction* ActiveProduction = ProductionSettings->GetActiveProduction())
-	{
-		ProductionSettings->RemoveAssetNaming(ActiveProduction->ProductionID, AssetNaming->Class);
-	}
+	ProductionSettings->RemoveAssetNaming(ProductionSettings->GetActiveProductionID(), AssetNaming->Class);
 
 	AssetNaming->Class = SelectedClass;
 
@@ -357,10 +356,7 @@ void SAssetNamingRow::SetDefaultAssetName(const FText& InText, ETextCommit::Type
 	AssetNaming->DefaultName = InText.ToString();
 
 	UProductionSettings* ProductionSettings = GetMutableDefault<UProductionSettings>();
-	if (FCinematicProduction* ActiveProduction = ProductionSettings->GetActiveProduction())
-	{
-		ProductionSettings->AddAssetNaming(ActiveProduction->ProductionID, AssetNaming->Class, AssetNaming->DefaultName);
-	}
+	ProductionSettings->AddAssetNaming(ProductionSettings->GetActiveProductionID(), AssetNaming->Class, AssetNaming->DefaultName);
 }
 
 bool SAssetNamingRow::ValidateDefaultAssetName(const FText& InText, FText& OutErrorMessage)
@@ -380,10 +376,7 @@ bool SAssetNamingRow::ValidateDefaultAssetName(const FText& InText, FText& OutEr
 FReply SAssetNamingRow::DeleteRow()
 {
 	UProductionSettings* ProductionSettings = GetMutableDefault<UProductionSettings>();
-	if (FCinematicProduction* ActiveProduction = ProductionSettings->GetActiveProduction())
-	{
-		ProductionSettings->RemoveAssetNaming(ActiveProduction->ProductionID, AssetNaming->Class);
-	}
+	ProductionSettings->RemoveAssetNaming(ProductionSettings->GetActiveProductionID(), AssetNaming->Class);
 
 	OnDeleteRow.ExecuteIfBound();
 
diff --git a/Engine/Plugins/VirtualProduction/Shots/Source/ShotsEditor/Private/UI/SFolderHierarchyPanel.cpp b/Engine/Plugins/VirtualProduction/Shots/Source/ShotsEditor/Private/UI/SFolderHierarchyPanel.cpp
--- a/Engine/Plugins/VirtualProduction/Shots/Source/ShotsEditor/Private/UI/SFolderHierarchyPanel.cpp
+++ b/Engine/Plugins/VirtualProduction/Shots/Source/ShotsEditor/Private/UI/SFolderHierarchyPanel.cpp
@@ -130,7 +130,8 @@ void SFolderHierarchyPanel::Construct(const FArguments& InArgs)
 										.IsEnabled_Lambda([]() -> bool
 											{
 												const UProductionSettings* const ProductionSettings = GetDefault<UProductionSettings>();
-												return (ProductionSettings->GetActiveProduction() != nullptr);
+												TOptional<const FCinematicProduction> ActiveProduction = ProductionSettings->GetActiveProduction();
+												return ActiveProduction.IsSet();
 											}
 										)
 
@@ -436,13 +437,13 @@ void SFolderHierarchyPanel::BuildTreeFromProductionTemplate()
 	FolderItemsSource.Add(RootItem);
 
 	const UProductionSettings* const ProductionSettings = GetMutableDefault<UProductionSettings>();
-	const FCinematicProduction* ActiveProduction = ProductionSettings->GetActiveProduction();
-	if (ActiveProduction)
+	TOptional<const FCinematicProduction> ActiveProduction = ProductionSettings->GetActiveProduction();
+	if (ActiveProduction.IsSet())
 	{
 		FAssetRegistryModule& AssetRegistryModule = FModuleManager::GetModuleChecked<FAssetRegistryModule>("AssetRegistry");
 
 		// Sort the folder paths from the active production's template alphabetically to put them in the order they should appear in the tree
-		TArray<FFolderTemplate> ProductionTemplate = ActiveProduction->TemplateFolders;
+		TArray<FFolderTemplate> ProductionTemplate = ActiveProduction.GetValue().TemplateFolders;
 		CachedInitialState = ProductionTemplate;
 
 		ProductionTemplate.Sort([](const FFolderTemplate& A, const FFolderTemplate& B) { return A.InternalPath < B.InternalPath; });
@@ -524,11 +525,9 @@ FReply SFolderHierarchyPanel::OnAddFolderToTemplate()
 
 	// Add the new folder to the active production's template
 	UProductionSettings* ProductionSettings = GetMutableDefault<UProductionSettings>();
-	const FCinematicProduction* ActiveProduction = ProductionSettings->GetActiveProduction();
-	check(ActiveProduction);
 
 	const bool bCreateIfMissing = (NewFolder->Status == ETemplateFolderStatus::MissingDoNotCreate) ? false : true;
-	ProductionSettings->AddTemplateFolder(ActiveProduction->ProductionID, NewFolder->Path.GetInternalPathString(), bCreateIfMissing);
+	ProductionSettings->AddTemplateFolder(ProductionSettings->GetActiveProductionID(), NewFolder->Path.GetInternalPathString(), bCreateIfMissing);
 
 	TreeView->SetItemExpansion(SelectedTreeItem, true);
 	TreeView->RequestTreeRefresh();
@@ -542,10 +541,12 @@ FReply SFolderHierarchyPanel::OnCreateTemplateFolders()
 	CreateFolderFromTemplateRecursive(RootItem);
 
 	const UProductionSettings* const ProductionSettings = GetMutableDefault<UProductionSettings>();
-	const FCinematicProduction* ActiveProduction = ProductionSettings->GetActiveProduction();
-	check(ActiveProduction);
-	
-	CachedInitialState = ActiveProduction->TemplateFolders;
+	TOptional<const FCinematicProduction> ActiveProduction = ProductionSettings->GetActiveProduction();
+
+	if (ActiveProduction.IsSet())
+	{
+		CachedInitialState = ActiveProduction.GetValue().TemplateFolders;
+	}
 
 	TreeView->RequestTreeRefresh();
 	SetExpansionRecursive(RootItem, true);
@@ -557,10 +558,7 @@ FReply SFolderHierarchyPanel::OnResetTemplateChanges()
 {
 	// Reset the active production's template to the latest cached state, then rebuild the tree view
 	UProductionSettings* ProductionSettings = GetMutableDefault<UProductionSettings>();
-	const FCinematicProduction* ActiveProduction = ProductionSettings->GetActiveProduction();
-	check(ActiveProduction);
-
-	ProductionSettings->SetTemplateFolderHierarchy(ActiveProduction->ProductionID, CachedInitialState);
+	ProductionSettings->SetTemplateFolderHierarchy(ProductionSettings->GetActiveProductionID(), CachedInitialState);
 
 	BuildTreeFromProductionTemplate();
 
@@ -647,13 +645,11 @@ void SFolderHierarchyPanel::SetTemplateFolderName(const FText& InText, ETextComm
 
 	// Remove the old path from the active production's template, and add the new path instead
 	UProductionSettings* ProductionSettings = GetMutableDefault<UProductionSettings>();
-	FCinematicProduction* ActiveProduction = ProductionSettings->GetActiveProduction();
-	check(ActiveProduction);
 
-	ProductionSettings->RemoveTemplateFolder(ActiveProduction->ProductionID, OldPath);
+	ProductionSettings->RemoveTemplateFolder(ProductionSettings->GetActiveProductionID(), OldPath);
 
 	const bool bCreateIfMissing = (TreeItem->Status == ETemplateFolderStatus::MissingDoNotCreate) ? false : true;
-	ProductionSettings->AddTemplateFolder(ActiveProduction->ProductionID, NewPath, bCreateIfMissing);
+	ProductionSettings->AddTemplateFolder(ProductionSettings->GetActiveProductionID(), NewPath, bCreateIfMissing);
 
 	// Update the name input tree item and the path of all of its children (recursively)
 	TreeItem->Path.SetPathFromString(NewPath, EContentBrowserPathType::Internal);
@@ -769,10 +765,7 @@ void SFolderHierarchyPanel::RemoveFolderFromTemplateRecursive(TSharedPtr<FTempla
 	}
 
 	UProductionSettings* ProductionSettings = GetMutableDefault<UProductionSettings>();
-	const FCinematicProduction* ActiveProduction = ProductionSettings->GetActiveProduction();
-	check(ActiveProduction);
-
-	ProductionSettings->RemoveTemplateFolder(ActiveProduction->ProductionID, TreeItem->Path.GetInternalPathString());
+	ProductionSettings->RemoveTemplateFolder(ProductionSettings->GetActiveProductionID(), TreeItem->Path.GetInternalPathString());
 
 	TreeItem->Children.Empty();
 }
@@ -782,8 +775,6 @@ void SFolderHierarchyPanel::SetChildrenPathRecursive(TSharedPtr<FTemplateFolderT
 	FAssetRegistryModule& AssetRegistryModule = FModuleManager::GetModuleChecked<FAssetRegistryModule>("AssetRegistry");
 
 	UProductionSettings* ProductionSettings = GetMutableDefault<UProductionSettings>();
-	const FCinematicProduction* ActiveProduction = ProductionSettings->GetActiveProduction();
-	check(ActiveProduction);
 
 	for (TSharedPtr<FTemplateFolderTreeItem>& Child : TreeItem->Children)
 	{
@@ -809,10 +800,10 @@ void SFolderHierarchyPanel::SetChildrenPathRecursive(TSharedPtr<FTemplateFolderT
 		}
 
 		// Remove the old path from the active production's template, and add the new path instead
-		ProductionSettings->RemoveTemplateFolder(ActiveProduction->ProductionID, OldChildPath);
+		ProductionSettings->RemoveTemplateFolder(ProductionSettings->GetActiveProductionID(), OldChildPath);
 
 		const bool bCreateIfMissing = (Child->Status == ETemplateFolderStatus::MissingDoNotCreate) ? false : true;
-		ProductionSettings->AddTemplateFolder(ActiveProduction->ProductionID, NewChildPath, bCreateIfMissing);
+		ProductionSettings->AddTemplateFolder(ProductionSettings->GetActiveProductionID(), NewChildPath, bCreateIfMissing);
 
 		SetChildrenPathRecursive(Child, NewChildPath);
 	}
@@ -823,14 +814,12 @@ void SFolderHierarchyPanel::SetTemplateFolderStatusRecursive(TSharedPtr<FTemplat
 	TreeItem->Status = NewStatus;
 
 	UProductionSettings* ProductionSettings = GetMutableDefault<UProductionSettings>();
-	const FCinematicProduction* ActiveProduction = ProductionSettings->GetActiveProduction();
-	check(ActiveProduction);
 
 	const FString Path = TreeItem->Path.GetInternalPathString();
-	ProductionSettings->RemoveTemplateFolder(ActiveProduction->ProductionID, Path);
+	ProductionSettings->RemoveTemplateFolder(ProductionSettings->GetActiveProductionID(), Path);
 
 	const bool bCreateIfMissing = (TreeItem->Status == ETemplateFolderStatus::MissingDoNotCreate) ? false : true;
-	ProductionSettings->AddTemplateFolder(ActiveProduction->ProductionID, Path, bCreateIfMissing);
+	ProductionSettings->AddTemplateFolder(ProductionSettings->GetActiveProductionID(), Path, bCreateIfMissing);
 
 	for (TSharedPtr<FTemplateFolderTreeItem>& Child : TreeItem->Children)
 	{
diff --git a/Engine/Plugins/VirtualProduction/Shots/Source/ShotsEditor/Private/UI/SNamingTokensPanel.cpp b/Engine/Plugins/VirtualProduction/Shots/Source/ShotsEditor/Private/UI/SNamingTokensPanel.cpp
--- a/Engine/Plugins/VirtualProduction/Shots/Source/ShotsEditor/Private/UI/SNamingTokensPanel.cpp
+++ b/Engine/Plugins/VirtualProduction/Shots/Source/ShotsEditor/Private/UI/SNamingTokensPanel.cpp
@@ -298,7 +298,8 @@ TSharedRef<ITableRow> SNamingTokensPanel::OnGenerateNamingTokenNamespaceRow(TSha
 	auto IsActiveProductionValid = []() -> bool
 		{
 			const UProductionSettings* const ProductionSettings = GetDefault<UProductionSettings>();
-			return (ProductionSettings->GetActiveProduction() != nullptr);
+			TOptional<const FCinematicProduction> ActiveProduction = ProductionSettings->GetActiveProduction();
+			return ActiveProduction.IsSet();
 		};
 
 	return SNew(STableRow<TSharedPtr<FString>>, OwnerTable)
@@ -331,25 +332,23 @@ TSharedRef<ITableRow> SNamingTokensPanel::OnGenerateNamingTokenNamespaceRow(TSha
 void SNamingTokensPanel::OnNamespaceChecked(ECheckBoxState CheckBoxState, TSharedPtr<FString> Namespace)
 {
 	UProductionSettings* ProductionSettings = GetMutableDefault<UProductionSettings>();
-	if (FCinematicProduction* ActiveProduction = ProductionSettings->GetActiveProduction())
+	if (CheckBoxState == ECheckBoxState::Checked)
 	{
-		if (CheckBoxState == ECheckBoxState::Checked)
-		{
-			ProductionSettings->RemoveNamespaceFromDenyList(ActiveProduction->ProductionID, *Namespace);
-		}
-		else if (CheckBoxState == ECheckBoxState::Unchecked)
-		{
-			ProductionSettings->AddNamespaceToDenyList(ActiveProduction->ProductionID, *Namespace);
-		}
+		ProductionSettings->RemoveNamespaceFromDenyList(ProductionSettings->GetActiveProductionID(), *Namespace);
+	}
+	else if (CheckBoxState == ECheckBoxState::Unchecked)
+	{
+		ProductionSettings->AddNamespaceToDenyList(ProductionSettings->GetActiveProductionID(), *Namespace);
 	}
 }
 
 ECheckBoxState SNamingTokensPanel::GetNamespaceCheckBoxState(TSharedPtr<FString> Namespace) const
 {
 	const UProductionSettings* const ProductionSettings = GetDefault<UProductionSettings>();
-	if (FCinematicProduction* ActiveProduction = ProductionSettings->GetActiveProduction())
+	TOptional<const FCinematicProduction> ActiveProduction = ProductionSettings->GetActiveProduction();
+	if (ActiveProduction.IsSet())
 	{
-		if (ActiveProduction->NamingTokenNamespaceDenyList.Contains(*Namespace))
+		if (ActiveProduction.GetValue().NamingTokenNamespaceDenyList.Contains(*Namespace))
 		{
 			return ECheckBoxState::Unchecked;
 		}
diff --git a/Engine/Plugins/VirtualProduction/Shots/Source/ShotsEditor/Private/UI/SProductionListPanel.cpp b/Engine/Plugins/VirtualProduction/Shots/Source/ShotsEditor/Private/UI/SProductionListPanel.cpp
--- a/Engine/Plugins/VirtualProduction/Shots/Source/ShotsEditor/Private/UI/SProductionListPanel.cpp
+++ b/Engine/Plugins/VirtualProduction/Shots/Source/ShotsEditor/Private/UI/SProductionListPanel.cpp
@@ -2,6 +2,7 @@
 
 #include "SProductionListPanel.h"
 
+#include "Algo/Contains.h"
 #include "DesktopPlatformModule.h"
 #include "Framework/Application/MenuStack.h"
 #include "Framework/Application/SlateApplication.h"
@@ -34,7 +35,8 @@ void SProductionListPanel::Construct(const FArguments& InArgs)
 			UProductionSettings* ProductionSettings = GetMutableDefault<UProductionSettings>();
 			ProductionSettings->AddProduction();
 
-			MostRecentProductionID = ProductionSettings->Productions.Last().ProductionID;
+			const TArray<FCinematicProduction> Productions = ProductionSettings->GetProductions();
+			MostRecentProductionID = Productions.Last().ProductionID;
 			return FReply::Handled();
 		};
 
@@ -168,7 +170,8 @@ void SProductionListPanel::UpdateProductionList()
 	ProductionList.Empty();
 
 	const UProductionSettings* const ProductionSettings = GetDefault<UProductionSettings>();
-	for (const FCinematicProduction& Production : ProductionSettings->Productions)
+	const TArray<FCinematicProduction> Productions = ProductionSettings->GetProductions();
+	for (const FCinematicProduction& Production : Productions)
 	{
 		ProductionList.Add(MakeShared<FCinematicProduction>(Production));
 	}
@@ -290,7 +293,8 @@ FReply SProductionListPanel::ImportProduction()
 				ImportedProduction->ProductionID = FGuid::NewGuid();
 
 				// If the name of the production being imported already exists in the list of productions, prompt the user to confirm before actually adding it
-				const bool bProductionNameAlreadyExists = ProductionSettings->Productions.ContainsByPredicate([ProductionName = ImportedProduction->ProductionName](const FCinematicProduction& Other) { return Other.ProductionName == ProductionName; });
+				const TArray<FCinematicProduction> Productions = ProductionSettings->GetProductions();			
+				const bool bProductionNameAlreadyExists = Algo::ContainsBy(Productions, ImportedProduction->ProductionName, &FCinematicProduction::ProductionName);
 
 				bool bShouldAddProduction = true;
 				if (bProductionNameAlreadyExists)
@@ -328,21 +332,21 @@ FReply SProductionListPanel::ImportProduction()
 FReply SProductionListPanel::ExportProduction(TSharedPtr<FCinematicProduction> InItem)
 {
 	UProductionSettings* ProductionSettings = GetMutableDefault<UProductionSettings>();
-	const FCinematicProduction* ProductionToExport = ProductionSettings->GetProduction(InItem->ProductionID);
-	if (!ProductionToExport)
+	TOptional<const FCinematicProduction> ProductionToExport = ProductionSettings->GetProduction(InItem->ProductionID);
+	if (!ProductionToExport.IsSet())
 	{
 		return FReply::Handled();
 	}
 
 	// Convert the production settings to a json object that can be written to .json file on disk
-	if (const TSharedPtr<FJsonObject>& JsonObject = FJsonObjectConverter::UStructToJsonObject<FCinematicProduction>(*ProductionToExport))
+	if (const TSharedPtr<FJsonObject>& JsonObject = FJsonObjectConverter::UStructToJsonObject<FCinematicProduction>(ProductionToExport.GetValue()))
 	{
 		// Prompt the user to choose a location to save the production setting .json file
 		IDesktopPlatform* DesktopPlatform = FDesktopPlatformModule::Get();
 		const void* ParentWindowHandle = FSlateApplication::Get().FindBestParentWindowHandleForDialogs(nullptr);
 		const FString Title = TEXT("Export Production Settings");
 		const FString DefaultPath = FPaths::ProjectSavedDir();
-		const FString DefaultFile = ProductionToExport->ProductionName + ".json";
+		const FString DefaultFile = ProductionToExport.GetValue().ProductionName + ".json";
 		const FString FileTypes = TEXT("json|*.json");
 		const uint32 FileFlags = 0;
 
@@ -417,9 +421,10 @@ void SProductionListButton::Construct(const FArguments& InArgs)
 	auto GetProductionName = [this]() -> FText
 		{
 			const UProductionSettings* const ProductionSettings = GetDefault<UProductionSettings>();
-			if (const FCinematicProduction* Production = ProductionSettings->GetProduction(ProductionID.Get()))
+			TOptional<const FCinematicProduction> Production = ProductionSettings->GetProduction(ProductionID.Get());
+			if (Production.IsSet())
 			{
-				return FText::FromString(Production->ProductionName);
+				return FText::FromString(Production.GetValue().ProductionName);
 			}
 			return FText::GetEmpty();
 		};
@@ -442,7 +447,8 @@ void SProductionListButton::Construct(const FArguments& InArgs)
 			}
 
 			bool bIsNameValid = true;
-			for (const FCinematicProduction& Production : ProductionSettings->Productions)
+			const TArray<FCinematicProduction> Productions = ProductionSettings->GetProductions();
+			for (const FCinematicProduction& Production : Productions)
 			{
 				if (Production.ProductionID == ProductionID.Get())
 				{
diff --git a/Engine/Plugins/VirtualProduction/Shots/Source/ShotsEditor/Private/UI/SSequencerSettingsPanel.cpp b/Engine/Plugins/VirtualProduction/Shots/Source/ShotsEditor/Private/UI/SSequencerSettingsPanel.cpp
--- a/Engine/Plugins/VirtualProduction/Shots/Source/ShotsEditor/Private/UI/SSequencerSettingsPanel.cpp
+++ b/Engine/Plugins/VirtualProduction/Shots/Source/ShotsEditor/Private/UI/SSequencerSettingsPanel.cpp
@@ -22,7 +22,8 @@ void SSequencerSettingsPanel::Construct(const FArguments& InArgs)
 	auto IsActiveProductionValid = []() -> bool
 		{
 			const UProductionSettings* const ProductionSettings = GetDefault<UProductionSettings>();
-			return (ProductionSettings->GetActiveProduction() != nullptr);
+			TOptional<const FCinematicProduction> ActiveProduction = ProductionSettings->GetActiveProduction();
+			return ActiveProduction.IsSet();
 		};
 
 	// Get a text representation of the active subsequence priority to display in the combo button below
