ShotManagement: Add a feature to create Cine Assembly assets from Cine Assembly Schemas.
Adds an asset picker to the Cine Assembly config window where the user can select a schema as the base for their assembly.
Cine Assemblies also inherit the asset display name and thumbnail from their base schema.

#rnx

#virtualized

diff --git a/Engine/Plugins/VirtualProduction/Shots/Source/ShotsEditor/Private/AssetDefinition_CineAssembly.cpp b/Engine/Plugins/VirtualProduction/Shots/Source/ShotsEditor/Private/AssetDefinition_CineAssembly.cpp
--- a/Engine/Plugins/VirtualProduction/Shots/Source/ShotsEditor/Private/AssetDefinition_CineAssembly.cpp
+++ b/Engine/Plugins/VirtualProduction/Shots/Source/ShotsEditor/Private/AssetDefinition_CineAssembly.cpp
@@ -15,6 +15,18 @@ FText UAssetDefinition_CineAssembly::GetAssetDisplayName() const
 	return LOCTEXT("AssetTypeActions_CineAssembly", "Cine Assembly");
 }
 
+FText UAssetDefinition_CineAssembly::GetAssetDisplayName(const FAssetData& AssetData) const
+{
+	const FAssetDataTagMapSharedView::FFindTagResult AssemblyType = AssetData.TagsAndValues.FindTag(UCineAssembly::AssetRegistryTag_AssemblyType);
+
+	if (AssemblyType.IsSet())
+	{
+		return FText::FromString(AssemblyType.GetValue());
+	}
+
+	return GetAssetDisplayName();
+}
+
 FText UAssetDefinition_CineAssembly::GetAssetDescription(const FAssetData& AssetData) const
 {
 	if (const UCineAssembly* const CineAssembly = Cast<UCineAssembly>(AssetData.GetAsset()))
@@ -44,6 +56,19 @@ TConstArrayView<FAssetCategoryPath> UAssetDefinition_CineAssembly::GetAssetCateg
 	return Categories;
 }
 
+const FSlateBrush* UAssetDefinition_CineAssembly::GetThumbnailBrush(const FAssetData& InAssetData, const FName InClassName) const
+{
+	if (const UCineAssembly* const CineAssembly = Cast<UCineAssembly>(InAssetData.GetAsset()))
+	{
+		if (const UCineAssemblySchema* Schema = CineAssembly->GetSchema())
+		{
+			return Schema->GetThumbnailBrush();
+		}
+	}
+
+	return nullptr;
+}
+
 EAssetCommandResult UAssetDefinition_CineAssembly::OpenAssets(const FAssetOpenArgs& OpenArgs) const
 {
 	return Super::OpenAssets(OpenArgs);
diff --git a/Engine/Plugins/VirtualProduction/Shots/Source/ShotsEditor/Private/AssetDefinition_CineAssembly.h b/Engine/Plugins/VirtualProduction/Shots/Source/ShotsEditor/Private/AssetDefinition_CineAssembly.h
--- a/Engine/Plugins/VirtualProduction/Shots/Source/ShotsEditor/Private/AssetDefinition_CineAssembly.h
+++ b/Engine/Plugins/VirtualProduction/Shots/Source/ShotsEditor/Private/AssetDefinition_CineAssembly.h
@@ -15,10 +15,12 @@ class UAssetDefinition_CineAssembly : public UAssetDefinition_LevelSequence
 protected:
 	// Begin UAssetDefinition Interface
 	virtual FText GetAssetDisplayName() const override;
+	virtual FText GetAssetDisplayName(const FAssetData& AssetData) const override;
 	virtual FText GetAssetDescription(const FAssetData& AssetData) const override;
 	virtual TSoftClassPtr<UObject> GetAssetClass() const override;
 	virtual FLinearColor GetAssetColor() const override;
 	virtual TConstArrayView<FAssetCategoryPath> GetAssetCategories() const override;
+	virtual const FSlateBrush* GetThumbnailBrush(const FAssetData& InAssetData, const FName InClassName) const override;
 	virtual TArray<FAssetData> PrepareToActivateAssets(const FAssetActivateArgs& ActivateArgs) const override;
 	virtual EAssetCommandResult OpenAssets(const FAssetOpenArgs& OpenArgs) const override;
 	// End UAssetDefinition Interface
diff --git a/Engine/Plugins/VirtualProduction/Shots/Source/ShotsEditor/Private/CineAssembly.cpp b/Engine/Plugins/VirtualProduction/Shots/Source/ShotsEditor/Private/CineAssembly.cpp
--- a/Engine/Plugins/VirtualProduction/Shots/Source/ShotsEditor/Private/CineAssembly.cpp
+++ b/Engine/Plugins/VirtualProduction/Shots/Source/ShotsEditor/Private/CineAssembly.cpp
@@ -3,10 +3,32 @@
 #include "CineAssembly.h"
 #include "UObject/AssetRegistryTagsContext.h"
 
+#define LOCTEXT_NAMESPACE "CineAssembly"
+
+const FName UCineAssembly::AssetRegistryTag_AssemblyType = "AssemblyType";
+
+const UCineAssemblySchema* UCineAssembly::GetSchema() const
+{
+	return BaseSchema.Get();
+}
+
+void UCineAssembly::SetSchema(UCineAssemblySchema* InSchema)
+{
+	// If the assembly already has a base schema set, do not overwrite it
+	if (BaseSchema)
+	{
+		return;
+	}
+	BaseSchema = InSchema;
+}
+
 void UCineAssembly::GetAssetRegistryTags(FAssetRegistryTagsContext Context) const
 {
 	Super::GetAssetRegistryTags(Context);
 
+	const FString AssemblyType = BaseSchema ? BaseSchema->SchemaName : TEXT("");
+	Context.AddTag(FAssetRegistryTag(AssetRegistryTag_AssemblyType, AssemblyType, FAssetRegistryTag::ETagType::TT_Alphabetical, FAssetRegistryTag::TD_None));
+
 	for (const TPair<FName, FString>& Pair : CustomMetadata)
 	{
 		if (!Pair.Key.IsNone() && !Pair.Value.IsEmpty())
@@ -15,3 +37,15 @@ void UCineAssembly::GetAssetRegistryTags(FAssetRegistryTagsContext Context) cons
 		}
 	}
 }
+
+void UCineAssembly::GetAssetRegistryTagMetadata(TMap<FName, FAssetRegistryTagMetadata>& OutMetadata) const
+{
+	Super::GetAssetRegistryTagMetadata(OutMetadata);
+
+	OutMetadata.Add(AssetRegistryTag_AssemblyType, FAssetRegistryTagMetadata()
+		.SetDisplayName(LOCTEXT("AssemblyType_Label", "AssemblyType"))
+		.SetTooltip(LOCTEXT("AssemblyType_Tooltip", "The assembly type of this instance"))
+	);
+}
+
+#undef LOCTEXT_NAMESPACE
diff --git a/Engine/Plugins/VirtualProduction/Shots/Source/ShotsEditor/Private/CineAssembly.h b/Engine/Plugins/VirtualProduction/Shots/Source/ShotsEditor/Private/CineAssembly.h
--- a/Engine/Plugins/VirtualProduction/Shots/Source/ShotsEditor/Private/CineAssembly.h
+++ b/Engine/Plugins/VirtualProduction/Shots/Source/ShotsEditor/Private/CineAssembly.h
@@ -4,6 +4,8 @@
 
 #include "LevelSequence.h"
 
+#include "CineAssemblySchema.h"
+
 #include "CineAssembly.generated.h"
 
 /** 
@@ -19,8 +21,16 @@ public:
 
 	//~ Begin UObject Interface
 	virtual void GetAssetRegistryTags(FAssetRegistryTagsContext Context) const override;
+	virtual void GetAssetRegistryTagMetadata(TMap<FName, FAssetRegistryTagMetadata>& OutMetadata) const override;
 	//~ End UObject Interface
 
+	/** Get the base schema for this assembly */
+	const UCineAssemblySchema* GetSchema() const;
+
+	/** Assign the input schema to this assembly (if one has not already been assigned) */
+	void SetSchema(UCineAssemblySchema* InSchema);
+
+public:
 	/** The level to open before opening this asset in Sequencer */
 	UPROPERTY(EditAnywhere, Category = "TargetLevel", meta = (AllowedClasses = "/Script/Engine.World"))
 	FSoftObjectPath Level;
@@ -32,4 +42,12 @@ public:
 	/** User-facing notes about this assembly asset */
 	UPROPERTY()
 	FString AssemblyNote;
+
+	/** The asset registry tag that contains the assembly type information */
+	static const FName AssetRegistryTag_AssemblyType;
+
+private:
+	/** The schema that was used as a base when creating this assembly (can be null if no schema was used) */
+	UPROPERTY()
+	TObjectPtr<UCineAssemblySchema> BaseSchema;
 };
diff --git a/Engine/Plugins/VirtualProduction/Shots/Source/ShotsEditor/Private/CineAssemblySchema.cpp b/Engine/Plugins/VirtualProduction/Shots/Source/ShotsEditor/Private/CineAssemblySchema.cpp
--- a/Engine/Plugins/VirtualProduction/Shots/Source/ShotsEditor/Private/CineAssemblySchema.cpp
+++ b/Engine/Plugins/VirtualProduction/Shots/Source/ShotsEditor/Private/CineAssemblySchema.cpp
@@ -2,9 +2,18 @@
 
 #include "CineAssemblySchema.h"
 
+#include "AssetRegistry/AssetRegistryModule.h"
 #include "AssetToolsModule.h"
 #include "Modules/ModuleManager.h"
 
+const FSoftObjectPath UCineAssemblySchema::DefaultThumbnailPath = FSoftObjectPath(TEXT("/Shots/DefaultSchemaThumbnail.DefaultSchemaThumbnail"));
+
+UCineAssemblySchema::UCineAssemblySchema()
+{
+	ThumbnailBrush = MakeShared<FSlateBrush>();
+	ThumbnailBrush->ImageSize = FVector2D(64.0f, 64.0f);
+}
+
 void UCineAssemblySchema::PostLoad()
 {
 	Super::PostLoad();
@@ -18,14 +27,20 @@ void UCineAssemblySchema::PostEditChangeProperty(struct FPropertyChangedEvent& P
 
 	if (PropertyName == GET_MEMBER_NAME_CHECKED(UCineAssemblySchema, ThumbnailImage))
 	{
-		// We mark the new thumbnail as pending until the next time that the brush is requested from an external caller.
-		// This is because Slate holds on to the brush pointer, and assumes (in many places) that it will remain valid.
-		// So if we reset the brush pointer or reset its resource object now, we run the risk of crashing within Slate
-		// somewhere when it tries to access the stale brush.
-		bThumbnailPendingReset = true;
+		UpdateThumbnailBrush();
 	}
 }
 
+const FSlateBrush* UCineAssemblySchema::GetThumbnailBrush() const
+{
+	return ThumbnailBrush.Get();
+}
+
+bool UCineAssemblySchema::SupportsRename() const
+{
+	return bSupportsRename;
+}
+
 void UCineAssemblySchema::RenameAsset()
 {
 	// The default behavior for schema assets is to not allow renaming from Content Browser.
@@ -45,35 +60,17 @@ void UCineAssemblySchema::RenameAsset()
 	bSupportsRename = false;
 }
 
-bool UCineAssemblySchema::SupportsRename()
-{
-	return bSupportsRename;
-}
-
 void UCineAssemblySchema::UpdateThumbnailBrush()
 {
 	if (ThumbnailImage)
 	{
-		if (!ThumbnailBrush.IsValid())
-		{
-			ThumbnailBrush = MakeShareable(new FSlateBrush());
-		}
-
 		ThumbnailBrush->SetResourceObject(ThumbnailImage);
-		ThumbnailBrush->ImageSize = FVector2D(64.0f, 64.0f);
 	}
 	else
 	{
-		ThumbnailBrush.Reset();
-	}
-}
+		FAssetRegistryModule& AssetRegistryModule = FModuleManager::GetModuleChecked<FAssetRegistryModule>("AssetRegistry");
+		const FAssetData ThumbnailAssetData = AssetRegistryModule.Get().GetAssetByObjectPath(DefaultThumbnailPath);
 
-FSlateBrush* UCineAssemblySchema::GetThumbnailBrush()
-{
-	if (bThumbnailPendingReset)
-	{
-		UpdateThumbnailBrush();
+		ThumbnailBrush->SetResourceObject(ThumbnailAssetData.GetAsset());
 	}
-
-	return ThumbnailBrush.Get();
 }
diff --git a/Engine/Plugins/VirtualProduction/Shots/Source/ShotsEditor/Private/CineAssemblySchema.h b/Engine/Plugins/VirtualProduction/Shots/Source/ShotsEditor/Private/CineAssemblySchema.h
--- a/Engine/Plugins/VirtualProduction/Shots/Source/ShotsEditor/Private/CineAssemblySchema.h
+++ b/Engine/Plugins/VirtualProduction/Shots/Source/ShotsEditor/Private/CineAssemblySchema.h
@@ -16,21 +16,21 @@ class UCineAssemblySchema : public UObject
 	GENERATED_BODY()
 
 public:
-	UCineAssemblySchema() = default;
+	UCineAssemblySchema();
 
 	//~ Begin UObject overrides
 	virtual void PostLoad() override;
 	virtual void PostEditChangeProperty(struct FPropertyChangedEvent& PropertyChangedEvent) override;
 	//~ End UObject overrides
 
-	/** Renames the underlying schema asset */
-	void RenameAsset();
+	/** Returns the thumbnail brush used by this schema */
+	const FSlateBrush* GetThumbnailBrush() const;
 
 	/** Whether or not this schema can be renamed */
-	bool SupportsRename();
+	bool SupportsRename() const;
 
-	/** Returns the thumbnail brush used by this schema */
-	FSlateBrush* GetThumbnailBrush();
+	/** Renames the underlying schema asset */
+	void RenameAsset();
 
 private:
 	/** Updates the thumbnail brush resource */
@@ -58,4 +58,7 @@ private:
 
 	/** Whether the schema asset supports renaming */
 	bool bSupportsRename = false;
+
+	/** Asset path of the default thumbnail to use for schema assets */
+	static const FSoftObjectPath DefaultThumbnailPath;
 };
diff --git a/Engine/Plugins/VirtualProduction/Shots/Source/ShotsEditor/Private/UI/CineAssembly/SCineAssemblyConfigWindow.cpp b/Engine/Plugins/VirtualProduction/Shots/Source/ShotsEditor/Private/UI/CineAssembly/SCineAssemblyConfigWindow.cpp
--- a/Engine/Plugins/VirtualProduction/Shots/Source/ShotsEditor/Private/UI/CineAssembly/SCineAssemblyConfigWindow.cpp
+++ b/Engine/Plugins/VirtualProduction/Shots/Source/ShotsEditor/Private/UI/CineAssembly/SCineAssemblyConfigWindow.cpp
@@ -2,7 +2,9 @@
 
 #include "SCineAssemblyConfigWindow.h"
 
+#include "ContentBrowserModule.h"
 #include "Editor.h"
+#include "IContentBrowserSingleton.h"
 #include "IStructureDetailsView.h"
 #include "Modules/ModuleManager.h"
 #include "PropertyEditorModule.h"
@@ -60,7 +62,7 @@ void SCineAssemblyConfigWindow::Construct(const FArguments& InArgs, UCineAssembl
 								.Orientation(Orient_Horizontal)
 								.PhysicalSplitterHandleSize(2.0f)
 
-								+SSplitter::Slot()
+								+ SSplitter::Slot()
 								.Value(0.7f)
 								[
 									MakeCineTemplatePanel()
@@ -92,24 +94,69 @@ void SCineAssemblyConfigWindow::Construct(const FArguments& InArgs, UCineAssembl
 
 TSharedRef<SWidget> SCineAssemblyConfigWindow::MakeCineTemplatePanel()
 {
-	// TODO: Fill in this stubbed out panel
-	return SNew(SVerticalBox)
+	FContentBrowserModule& ContentBrowserModule = FModuleManager::Get().LoadModuleChecked<FContentBrowserModule>("ContentBrowser");
+
+	// The asset picker will only display Cine Assembly Schema assets
+	FAssetPickerConfig Config;
+	Config.Filter.ClassPaths.Add(UCineAssemblySchema::StaticClass()->GetClassPathName());
+	Config.SelectionMode = ESelectionMode::Single;
+	Config.InitialAssetViewType = EAssetViewType::Tile;
+	Config.bFocusSearchBoxWhenOpened = false;
+	Config.bShowBottomToolbar = false;
+	Config.bAllowDragging = false;
+	Config.bAllowRename = false;
+	Config.bCanShowClasses = false;
+	Config.bCanShowFolders = false;
+	Config.bForceShowEngineContent = true;
+	Config.bForceShowPluginContent = true;
+	Config.bAddFilterUI = false;
+	Config.bShowPathInColumnView = true;
+	Config.OnAssetSelected = FOnAssetSelected::CreateSP(this, &SCineAssemblyConfigWindow::OnSchemaSelected);
+
+	const FName NoSchemaName = "No Schema";
+	FAssetData NoSchemaAssetData = FAssetData("/", "", NoSchemaName, FTopLevelAssetPath(""));
+
+	// Add a fake asset to the list (so that it appears as a tile in the asset picker) that represents a selection of no schema
+	Config.OnGetCustomSourceAssets = FOnGetCustomSourceAssets::CreateLambda([NoSchemaAssetData](const FARFilter& Filter, TArray<FAssetData>& OutAssets)
+		{
+			OutAssets.Add(NoSchemaAssetData);
+		});
 
-		+ SVerticalBox::Slot()
-		.FillHeight(1.0f)
-		.HAlign(HAlign_Center)
-		[
-			SNew(SHorizontalBox)
+	Config.InitialAssetSelection = NoSchemaAssetData;
 
-			+ SHorizontalBox::Slot()
-				.FillWidth(1.0f)
-				.VAlign(VAlign_Center)
-				[
-					SNew(STextBlock).Text(LOCTEXT("CineTemplatePanel", "CINE TEMPLATE PANEL"))
-				]
+	// The fake NoSchema asset should not display the normal asset tooltip, just a plain text-based tooltip describing what it is
+	Config.OnIsAssetValidForCustomToolTip = FOnIsAssetValidForCustomToolTip::CreateLambda([NoSchemaAssetData](FAssetData& AssetData)
+		{
+			return (AssetData == NoSchemaAssetData) ? true : false;
+		});
+
+	Config.OnGetCustomAssetToolTip = FOnGetCustomAssetToolTip::CreateLambda([](FAssetData& AssetData) -> TSharedRef<SToolTip>
+		{
+			return SNew(SToolTip).Text(LOCTEXT("NoSchemaToolTip", "Create a new assembly with no schema"));
+		});
+
+	return SNew(SBorder)
+		.BorderImage(FShotsEditorStyle::Get().GetBrush("ProductionWizard.PanelBackground"))
+		.Padding(16.0f)
+		[
+			ContentBrowserModule.Get().CreateAssetPicker(Config)
 		];
 }
 
+void SCineAssemblyConfigWindow::OnSchemaSelected(const FAssetData& AssetData)
+{
+	if (UCineAssemblySchema* CineAssemblySchema = Cast<UCineAssemblySchema>(AssetData.GetAsset()))
+	{
+		SelectedSchema = CineAssemblySchema;
+	}
+	else
+	{
+		SelectedSchema = nullptr;
+	}
+
+	CineAssemblyToConfigure->SetSchema(SelectedSchema);
+}
+
 TSharedRef<SWidget> SCineAssemblyConfigWindow::MakeInfoPanel()
 {
 	// TODO: Remove this when all of the tab content has been implemented
@@ -127,7 +174,7 @@ TSharedRef<SWidget> SCineAssemblyConfigWindow::MakeInfoPanel()
 
 	TabSwitcher = SNew(SWidgetSwitcher)
 		+ SWidgetSwitcher::Slot()
-			[ BuildStubWidget("OVERVIEW TAB CONTENT") ]
+			[ MakeOverviewWidget() ]
 
 		+ SWidgetSwitcher::Slot()
 			[ MakeDetailsWidget() ]
@@ -175,6 +222,76 @@ TSharedRef<SWidget> SCineAssemblyConfigWindow::MakeInfoPanel()
 		];
 }
 
+TSharedRef<SWidget> SCineAssemblyConfigWindow::MakeOverviewWidget()
+{
+	return SNew(SBorder)
+		.BorderImage(FShotsEditorStyle::Get().GetBrush("ProductionWizard.PanelBackground"))
+		.Padding(16.0f)
+		[
+			SNew(SVerticalBox)
+
+			// TODO: Populate this with an actual preview image sourced from the schema
+			+ SVerticalBox::Slot()
+				.Padding(0.0f, 0.0f, 0.0f, 16.0f)
+				.MinHeight(300.0f)
+				.MaxHeight(300.0f)
+				[
+					SNew(SBorder)
+						.BorderImage(FShotsEditorStyle::Get().GetBrush("Borders.Background"))
+						.HAlign(HAlign_Center)
+						.VAlign(VAlign_Center)
+						[
+							SNew(STextBlock).Text(LOCTEXT("PreviewImageText", "PREVIEW IMAGE"))
+						]
+				]
+
+			+ SVerticalBox::Slot()
+				.AutoHeight()
+				.Padding(0.0f, 0.0f, 0.0f, 4.0f)
+				[
+					SNew(STextBlock)
+						.AutoWrapText(true)
+						.Text_Lambda([this]() -> FText
+						{
+							const FText NoSchemaText = LOCTEXT("SchemaInstructions", "Choose a schema to use as the base for configuring your Cine Assembly, or proceed with no schema.");
+							return SelectedSchema ? FText::FromString(SelectedSchema->SchemaName) : NoSchemaText;
+						})
+				]
+
+			+ SVerticalBox::Slot()
+				.AutoHeight()
+				.Padding(0.0f, 0.0f, 0.0f, 8.0f)
+				[
+					SNew(STextBlock)
+						.Text(LOCTEXT("SchemClassName", "Cine Assembly Schema"))
+						.ColorAndOpacity(FSlateColor::UseSubduedForeground())
+						.Visibility_Lambda([this]() -> EVisibility 
+						{
+							return SelectedSchema ? EVisibility::Visible : EVisibility::Collapsed;
+						})
+				]
+
+			+ SVerticalBox::Slot()
+				.AutoHeight()
+				[
+					SNew(STextBlock)
+						.Text_Lambda([this]() -> FText
+						{
+							if (SelectedSchema)
+							{
+								return !SelectedSchema->Description.IsEmpty() ? FText::FromString(SelectedSchema->Description) : LOCTEXT("EmptyDescription", "No description");
+							}
+							return FText::GetEmpty();
+						})
+						.Visibility_Lambda([this]() -> EVisibility
+						{
+							return SelectedSchema ? EVisibility::Visible : EVisibility::Collapsed;
+						})
+
+				]
+		];
+}
+
 TSharedRef<SWidget> SCineAssemblyConfigWindow::MakeDetailsWidget()
 {
 	FPropertyEditorModule& PropertyEditor = FModuleManager::Get().LoadModuleChecked<FPropertyEditorModule>(TEXT("PropertyEditor"));
@@ -196,7 +313,6 @@ TSharedRef<SWidget> SCineAssemblyConfigWindow::MakeNotesWidget()
 		+ SVerticalBox::Slot()
 		.Padding(16.0f)
 		.AutoHeight()
-		.HAlign(HAlign_Left)
 		[
 			SNew(STextBlock)
 				.Text(LOCTEXT("NoteInstructions", "The following notes will be stored on the shot. This can also be edited later."))
diff --git a/Engine/Plugins/VirtualProduction/Shots/Source/ShotsEditor/Private/UI/CineAssembly/SCineAssemblyConfigWindow.h b/Engine/Plugins/VirtualProduction/Shots/Source/ShotsEditor/Private/UI/CineAssembly/SCineAssemblyConfigWindow.h
--- a/Engine/Plugins/VirtualProduction/Shots/Source/ShotsEditor/Private/UI/CineAssembly/SCineAssemblyConfigWindow.h
+++ b/Engine/Plugins/VirtualProduction/Shots/Source/ShotsEditor/Private/UI/CineAssembly/SCineAssemblyConfigWindow.h
@@ -5,6 +5,7 @@
 #include "Widgets/SCompoundWidget.h"
 
 #include "CineAssembly.h"
+#include "CineAssemblySchema.h"
 #include "Widgets/Layout/SWidgetSwitcher.h"
 #include "Widgets/SWindow.h"
 
@@ -41,18 +42,24 @@ private:
 	/** Creates the buttons on the bottom of the window */
 	TSharedRef<SWidget> MakeButtonsPanel();
 
+	/** Creates the widget to display for the Overview tab */
+	TSharedRef<SWidget> MakeOverviewWidget();
+
 	/** Creates the widget to display for the Details tab */
 	TSharedRef<SWidget> MakeDetailsWidget();
 
 	/** Creates the widget to display for the Notes tab */
 	TSharedRef<SWidget> MakeNotesWidget();
 
-	/** Cycles backwards through the available panels */
+	/** Closes the window and indicates that a new asset should be created by the asset factory */
 	FReply OnCreateAssetClicked();
 
-	/** Cycles forwards through the available panels */
+	/** Closes the window and indicates that no assets should be created by the asset factory */
 	FReply OnCancelClicked();
 
+	/** Updates the UI and CineAssembly properties based on the selected schema */
+	void OnSchemaSelected(const FAssetData& AssetData);
+
 private:
 	/** Switcher that controls which tab widget is currently visible */
 	TSharedPtr<SWidgetSwitcher> TabSwitcher;
@@ -63,6 +70,9 @@ private:
 	/** Transient object used only by this UI to configure the properties of the new asset that will get created by the Factory */
 	UCineAssembly* CineAssemblyToConfigure = nullptr;
 
+	/** The currently selected schema to use as a base for configuring the new cine assembly */
+	UCineAssemblySchema* SelectedSchema = nullptr;
+
 	/** Mode that configures the UI based on the intended user interactions with the Cine Assembly asset */
 	ECineAssemblyConfigMode Mode = ECineAssemblyConfigMode::CreateNew;
 };
