Shot Management: Add feature to Schemas to specify a list of subsequences that should get created when creating a new Cine Assembly from that Schema.
Subsequences are currently created as base Cine Assemblies (they do not have a schema).
The Assembly made from the Schema will automatically add a subsequence track. And that track will add the subsequence assemblies to itself.
#rnx

diff --git a/Engine/Plugins/VirtualProduction/Shots/Source/ShotsEditor/Private/CineAssemblyFactory.cpp b/Engine/Plugins/VirtualProduction/Shots/Source/ShotsEditor/Private/CineAssemblyFactory.cpp
--- a/Engine/Plugins/VirtualProduction/Shots/Source/ShotsEditor/Private/CineAssemblyFactory.cpp
+++ b/Engine/Plugins/VirtualProduction/Shots/Source/ShotsEditor/Private/CineAssemblyFactory.cpp
@@ -2,11 +2,13 @@
 
 #include "CineAssemblyFactory.h"
 
+#include "AssetToolsModule.h"
 #include "Framework/Application/SlateApplication.h"
 #include "Interfaces/IMainFrameModule.h"
 #include "Modules/ModuleManager.h"
 #include "MovieScene.h"
 #include "MovieSceneToolsProjectSettings.h"
+#include "Tracks/MovieSceneSubTrack.h"
 #include "UI/CineAssembly/SCineAssemblyConfigWindow.h"
 #include "UObject/Package.h"
 
@@ -34,14 +36,56 @@ UObject* UCineAssemblyFactory::FactoryCreateNew(UClass* Class, UObject* InParent
 		CineAssemblyToConfigure->Rename(*Name.ToString(), InParent);
 		CineAssemblyToConfigure->SetFlags(Flags);
 		CineAssemblyToConfigure->ClearFlags(RF_Transient);
+	}
+	else
+	{
+		// Procedural assembly creation that does not use the configuration window will hit this path
+		CineAssemblyToConfigure = NewObject<UCineAssembly>(InParent, Name, Flags);
+	}
 
-		// Do the same setup that ULevelSequence assets do when they are created
-		CineAssemblyToConfigure->Initialize();
+	// Do the same setup that ULevelSequence assets do when they are created
+	CineAssemblyToConfigure->Initialize();
 
-		const UMovieSceneToolsProjectSettings* MovieSceneToolsProjectSettings = GetDefault<UMovieSceneToolsProjectSettings>();
+	const UMovieSceneToolsProjectSettings* MovieSceneToolsProjectSettings = GetDefault<UMovieSceneToolsProjectSettings>();
 
-		FFrameRate TickResolution = CineAssemblyToConfigure->GetMovieScene()->GetTickResolution();
-		CineAssemblyToConfigure->GetMovieScene()->SetPlaybackRange((MovieSceneToolsProjectSettings->DefaultStartTime * TickResolution).FloorToFrame(), (MovieSceneToolsProjectSettings->DefaultDuration * TickResolution).FloorToFrame().Value);
+	FFrameRate TickResolution = CineAssemblyToConfigure->GetMovieScene()->GetTickResolution();
+	CineAssemblyToConfigure->GetMovieScene()->SetPlaybackRange((MovieSceneToolsProjectSettings->DefaultStartTime * TickResolution).FloorToFrame(), (MovieSceneToolsProjectSettings->DefaultDuration * TickResolution).FloorToFrame().Value);
+
+	// Create the subsequence assets associated with this assembly
+	if (const UCineAssemblySchema* Schema = CineAssemblyToConfigure->GetSchema())
+	{
+		if (Schema->Subsequences.Num() > 0)
+		{
+			// Add a subsequence track to the assembly's sequence
+			UMovieScene* MovieScene = CineAssemblyToConfigure->GetMovieScene();
+			UMovieSceneSubTrack* SubTrack = Cast<UMovieSceneSubTrack>(MovieScene->AddTrack(UMovieSceneSubTrack::StaticClass()));
+
+			IAssetTools& AssetTools = FModuleManager::GetModuleChecked<FAssetToolsModule>("AssetTools").Get();
+
+			// Get the path where the top-level assembly will be created so we can create the subsequences in the same location
+			// TODO: Once schemas support specifying the path for subsequences, this code should use those paths relative to this package path
+			const FAssetData AssemblyAssetData = FAssetData(CineAssemblyToConfigure);
+			const FString PackagePath = AssemblyAssetData.PackagePath.ToString();
+
+			// Create a new CineAssembly for each subsequence, set its playback range to match the parent sequence, and add it to the subsequence track
+			for (const FString& Subsequence : Schema->Subsequences)
+			{
+				// Before creating each subassembly, sanity check that each one will actually have a unique name (in case there are duplicates in the schema description)
+				FString UniquePackageName;
+				FString UniqueAssetName;
+				AssetTools.CreateUniqueAssetName(PackagePath / Subsequence, TEXT(""), UniquePackageName, UniqueAssetName);
+
+				UCineAssembly* SubAssembly = Cast<UCineAssembly>(AssetTools.CreateAsset(UniqueAssetName, PackagePath, UCineAssembly::StaticClass(), NewObject<UCineAssemblyFactory>()));
+
+				if (SubAssembly)
+				{
+					SubAssembly->Level = CineAssemblyToConfigure->Level;
+
+					const TRange<FFrameNumber> PlaybackRange = SubAssembly->GetMovieScene()->GetPlaybackRange();
+					SubTrack->AddSequence(SubAssembly, PlaybackRange.GetLowerBoundValue(), PlaybackRange.GetUpperBoundValue().Value - PlaybackRange.GetLowerBoundValue().Value);
+				}
+			}
+		}
 	}
 
 	return CineAssemblyToConfigure;
diff --git a/Engine/Plugins/VirtualProduction/Shots/Source/ShotsEditor/Private/CineAssemblySchema.h b/Engine/Plugins/VirtualProduction/Shots/Source/ShotsEditor/Private/CineAssemblySchema.h
--- a/Engine/Plugins/VirtualProduction/Shots/Source/ShotsEditor/Private/CineAssemblySchema.h
+++ b/Engine/Plugins/VirtualProduction/Shots/Source/ShotsEditor/Private/CineAssemblySchema.h
@@ -49,6 +49,10 @@ public:
 	UPROPERTY(EditAnywhere, Category = "Default")
 	TObjectPtr<UTexture2D> ThumbnailImage;
 
+	/** The names of the subsequences that should be created for assemblies that use this schema */
+	UPROPERTY()
+	TArray<FString> Subsequences;
+
 private:
 	/** Slate brush used to actually render the thumbnail image */
 	TSharedPtr<FSlateBrush> ThumbnailBrush;
diff --git a/Engine/Plugins/VirtualProduction/Shots/Source/ShotsEditor/Private/UI/CineAssembly/SCineAssemblySchemaWindow.cpp b/Engine/Plugins/VirtualProduction/Shots/Source/ShotsEditor/Private/UI/CineAssembly/SCineAssemblySchemaWindow.cpp
--- a/Engine/Plugins/VirtualProduction/Shots/Source/ShotsEditor/Private/UI/CineAssembly/SCineAssemblySchemaWindow.cpp
+++ b/Engine/Plugins/VirtualProduction/Shots/Source/ShotsEditor/Private/UI/CineAssembly/SCineAssemblySchemaWindow.cpp
@@ -2,10 +2,12 @@
 
 #include "SCineAssemblySchemaWindow.h"
 
+#include "Framework/MultiBox/MultiBoxBuilder.h"
 #include "IStructureDetailsView.h"
 #include "Modules/ModuleManager.h"
 #include "PropertyEditorModule.h"
 #include "ShotsEditorStyle.h"
+#include "Styling/StyleColors.h"
 #include "Widgets/Input/SButton.h"
 #include "Widgets/Input/SSegmentedControl.h"
 #include "Widgets/Layout/SBorder.h"
@@ -35,6 +37,12 @@ void SCineAssemblySchemaWindow::Construct(const FArguments& InArgs, UCineAssembl
 
 	GetOnWindowClosedEvent().AddSP(this, &SCineAssemblySchemaWindow::OnSchemaWindowClosed);
 
+	// Initialize the asset list with the current list saved in the schema
+	for (const FString& Subsequence : SchemaToConfigure->Subsequences)
+	{
+		AssetList.Add(MakeShared<FString>(Subsequence));
+	}
+
 	SWindow::Construct(SWindow::FArguments()
 		.Title(WindowTitle)
 		.SizingRule(ESizingRule::UserSized)
@@ -148,7 +156,7 @@ TSharedRef<SWidget> SCineAssemblySchemaWindow::MakeContentPanel()
 
 		+ SWidgetSwitcher::Slot()
 		[
-			BuildStubWidget("HIERARCHY TAB CONTENT")
+			MakeHierarchyTabContent()
 		];
 
 	return TabSwitcher.ToSharedRef();
@@ -210,6 +218,290 @@ TSharedRef<SWidget> SCineAssemblySchemaWindow::MakeAssetTabContent()
 	return DetailsView;
 }
 
+TSharedRef<SWidget> SCineAssemblySchemaWindow::MakeHierarchyTabContent()
+{
+	TSharedRef<SWidget> ContentHierarhcyWidget = SNew(SVerticalBox)
+
+		+ SVerticalBox::Slot()
+		.AutoHeight()
+		.Padding(0.0f, 0.0f, 0.0f, 4.0f)
+		[
+			SNew(STextBlock)
+				.Text(LOCTEXT("ContentHierarchyTitle", "Content Hierarchy"))
+				.Font(FShotsEditorStyle::Get().GetFontStyle("ProductionWizard.HeadingFont"))
+		]
+
+		+ SVerticalBox::Slot()
+		.AutoHeight()
+		.Padding(0.0f, 0.0f, 0.0f, 8.0f)
+		[
+			SNew(STextBlock)
+				.AutoWrapText(true)
+				.Text(LOCTEXT("ContentHierarchyText1", "Create a content hierarchy for your shot type. "
+					"This hierarchy will be automatically created when creating shots from this shot type."))
+		]
+
+		+ SVerticalBox::Slot()
+		.AutoHeight()
+		.Padding(0.0f, 0.0f, 0.0f, 16.0f)
+		[
+			SNew(STextBlock)
+				.AutoWrapText(true)
+				.Text(LOCTEXT("ContentHierarchyText2", "Assets being created or imported with this shot can be found to the right. "
+					"Any assets being created based on your timeline template must be assigned a save location before this shot type can be created."))
+		]
+
+		+ SVerticalBox::Slot()
+		.AutoHeight()
+		.HAlign(HAlign_Left)
+		.Padding(0.0f, 0.0f, 0.0f, 8.0f)
+		[
+			// TODO: Implement this add button to build out an actual tree view to drag assets from the asset list onto
+			SNew(SButton)
+				.ContentPadding(FMargin(2.0f))
+				[
+					SNew(SHorizontalBox)
+
+					+ SHorizontalBox::Slot()
+						.AutoWidth()
+						.Padding(0.0f, 0.0f, 4.0f, 0.0f)
+						[
+							SNew(SImage)
+								.Image(FAppStyle::Get().GetBrush("Icons.Plus"))
+								.ColorAndOpacity(FStyleColors::AccentGreen)
+						]
+
+					+ SHorizontalBox::Slot()
+						.AutoWidth()
+						[
+							SNew(STextBlock).Text(LOCTEXT("AddFolderButton", "Add"))
+						]
+		]		]
+
+		+ SVerticalBox::Slot()
+		.FillHeight(1.0f)
+		[
+			// TODO: Implement the tree view widget here for the folder hierarchy
+			SNew(SBorder)
+				.BorderImage(FShotsEditorStyle::Get().GetBrush("ProductionWizard.RecessedBackground"))
+		];
+
+	TSharedRef<SWidget> AssetListWidget = SNew(SVerticalBox)
+
+		+ SVerticalBox::Slot()
+		.AutoHeight()
+		.Padding(0.0f, 0.0f, 0.0f, 16.0f)
+		[
+			SNew(STextBlock)
+				.AutoWrapText(true)
+				.Text(LOCTEXT("AssetListInstructions", "The following assets will be created based on your Timeline Template. "
+				"Drag and drop them to a location in the hierarchy to determine where they should be saved when a new shot is created from this shot type"))
+		]
+
+		+ SVerticalBox::Slot()
+		.AutoHeight()
+		.HAlign(HAlign_Left)
+		.Padding(0.0f, 0.0f, 0.0f, 8.0f)
+		[
+			SNew(SButton)
+				.ContentPadding(FMargin(2.0f))
+				.OnClicked(this, &SCineAssemblySchemaWindow::OnAddAssetClicked)
+				[
+					SNew(SHorizontalBox)
+
+					+ SHorizontalBox::Slot()
+						.AutoWidth()
+						.Padding(0.0f, 0.0f, 4.0f, 0.0f)
+						[
+							SNew(SImage)
+								.Image(FAppStyle::Get().GetBrush("Icons.Plus"))
+								.ColorAndOpacity(FStyleColors::AccentGreen)
+						]
+
+					+ SHorizontalBox::Slot()
+						.AutoWidth()
+						[
+							SNew(STextBlock).Text(LOCTEXT("AddAssetButton", "Add"))
+						]
+		]		]
+
+		+ SVerticalBox::Slot()
+		.FillHeight(1.0f)
+		[
+			SNew(SBorder)
+				.BorderImage(FShotsEditorStyle::Get().GetBrush("ProductionWizard.RecessedBackground"))
+				[
+					SAssignNew(AssetListView, SListView<TSharedPtr<FString>>)
+						.ListItemsSource(&AssetList)
+						.SelectionMode(ESelectionMode::Single)
+						.OnGenerateRow(this, &SCineAssemblySchemaWindow::OnGenerateAssetRow)
+						.OnItemsRebuilt(this, &SCineAssemblySchemaWindow::OnAssetListRebuilt)
+						.OnContextMenuOpening(this, &SCineAssemblySchemaWindow::MakeAssetListContextMenu)
+						.OnMouseButtonDoubleClick(this, &SCineAssemblySchemaWindow::StartRename)
+				]
+		];
+
+	return SNew(SHorizontalBox)
+
+		+ SHorizontalBox::Slot()
+		.FillWidth(0.5f)
+		.Padding(16.0f)
+		[
+			ContentHierarhcyWidget
+		]
+
+		+ SHorizontalBox::Slot()
+		.AutoWidth()
+		[
+			SNew(SSeparator)
+				.Orientation(Orient_Vertical)
+				.Thickness(2.0f)
+		]
+
+		+ SHorizontalBox::Slot()
+		.FillWidth(0.5f)
+		.Padding(16.0f)
+		[
+			AssetListWidget
+		];
+}
+
+FReply SCineAssemblySchemaWindow::OnAddAssetClicked()
+{
+	// Add a new empty string to the list, which will be renamed by the user
+	AssetList.Add(MakeShared<FString>(TEXT("")));
+	AssetListView->RequestListRefresh();
+	return FReply::Handled();
+}
+
+TSharedRef<ITableRow> SCineAssemblySchemaWindow::OnGenerateAssetRow(TSharedPtr<FString> InItem, const TSharedRef<STableViewBase>& OwnerTable)
+{
+	TSharedRef<SSchemaAssetTableRow> Row = SNew(SSchemaAssetTableRow, OwnerTable)
+		.ShowSelection(true)
+		.Padding(FMargin(4.0f, 2.0f));
+
+	// Store a reference to the editable textblock in the row to easily set it to edit mode for renaming
+	Row->TextBlock = SNew(SInlineEditableTextBlock)
+		.Text_Lambda([InItem]() 
+		{ 
+			return FText::FromString(*InItem.Get()); 
+		})
+		.OnVerifyTextChanged(this, &SCineAssemblySchemaWindow::ValidateAssetName)
+		.OnTextCommitted_Lambda([this, InItem](const FText& InText, ETextCommit::Type InCommitType) 
+		{ 
+			*InItem = InText.ToString();
+		});
+
+	Row->SetContent(SNew(SBorder)
+		.Padding(8.0f)
+		.BorderImage(FShotsEditorStyle::Get().GetBrush("ProductionWizard.PanelBackground"))
+		[
+			SNew(SHorizontalBox)
+
+			+ SHorizontalBox::Slot()
+				.AutoWidth()
+				.VAlign(VAlign_Center)
+				.Padding(0.0f, 0.0f, 8.0f, 0.0f)
+				[
+					SNew(SImage).Image(FShotsEditorStyle::Get().GetBrush("Icons.Sequencer"))
+				]
+
+			+ SHorizontalBox::Slot()
+				.FillWidth(1.0f)
+				.VAlign(VAlign_Center)
+				[
+					Row->TextBlock.ToSharedRef()
+				]
+		]);
+
+	return Row;
+}
+
+void SCineAssemblySchemaWindow::OnAssetListRebuilt()
+{
+	if (AssetListView && AssetListView->GetItems().Num() > 0)
+	{
+		// Only trigger edit mode on the last item in the asset list, and only if it is an empty string so that it can be renamed
+		TSharedPtr<FString> LastItem = AssetListView->GetItems().Last();
+		if (LastItem->IsEmpty())
+		{
+			*LastItem = TEXT("NewAsset");
+			StartRename(LastItem);
+		}
+	}
+}
+
+TSharedPtr<SWidget> SCineAssemblySchemaWindow::MakeAssetListContextMenu()
+{
+	// The ListView uses single selection mode, so at most one item can ever be selected by the user
+	TArray<TSharedPtr<FString>> SelectedItems = AssetListView->GetSelectedItems();
+	if (SelectedItems.Num() == 1)
+	{
+		constexpr bool bCloseAfterMenuSelection = true;
+		FMenuBuilder MenuBuilder(bCloseAfterMenuSelection, nullptr);
+
+		TSharedPtr<FString>& SelectedItem = SelectedItems[0];
+
+		MenuBuilder.AddMenuEntry(
+			LOCTEXT("RenameAsset", "Rename"),
+			LOCTEXT("RenameAssetToolTip", "Rename"),
+			FSlateIcon(FShotsEditorStyle::StyleName, "Icons.AssetNaming"),
+			FUIAction(FExecuteAction::CreateLambda([this, SelectedItem]()
+				{
+					StartRename(SelectedItem);
+				})
+			),
+			NAME_None,
+			EUserInterfaceActionType::Button
+		);
+
+		MenuBuilder.AddMenuEntry(
+			LOCTEXT("DeleteAssetFromList", "Delete"),
+			LOCTEXT("DeleteAssetFromListToolTip", "Delete"),
+			FSlateIcon(FAppStyle::GetAppStyleSetName(), "Icons.Delete"),
+			FUIAction(FExecuteAction::CreateLambda([this, SelectedItem]()
+				{
+					const int32 IndexToRemove = AssetList.IndexOfByKey(SelectedItem);
+					if (AssetList.IsValidIndex(IndexToRemove))
+					{
+						AssetList.RemoveAt(IndexToRemove);
+					}
+					AssetListView->RequestListRefresh();
+				})
+			),
+			NAME_None,
+			EUserInterfaceActionType::Button
+		);
+
+		return MenuBuilder.MakeWidget();
+	}
+
+	return SNullWidget::NullWidget;
+}
+
+void SCineAssemblySchemaWindow::StartRename(TSharedPtr<FString> ItemToRename)
+{
+	// Get the row for the input item and put its textblock in edit mode so the user can rename the item
+	TSharedPtr<SSchemaAssetTableRow> Widget = StaticCastSharedPtr<SSchemaAssetTableRow>(AssetListView->WidgetFromItem(ItemToRename));
+
+	FSlateApplication::Get().SetKeyboardFocus(Widget->TextBlock.ToSharedRef());
+	Widget->TextBlock->EnterEditingMode();
+}
+
+bool SCineAssemblySchemaWindow::ValidateAssetName(const FText& InText, FText& OutErrorMessage)
+{
+	// An empty name is invalid
+	if (InText.IsEmpty())
+	{
+		OutErrorMessage = LOCTEXT("EmptyNameErrorMessage", "Please provide a name for the asset");
+		return false;
+	}
+
+	// Ensure that the name does not contain any characters that would be invalid for an asset name
+	// This matches the validation that would happen if the user was renaming an asset in the content browser
+	return FName::IsValidXName(InText.ToString(), INVALID_OBJECTNAME_CHARACTERS INVALID_LONGPACKAGE_CHARACTERS, &OutErrorMessage);
+}
+
 void SCineAssemblySchemaWindow::OnSchemaPropertiesChanged(const FPropertyChangedEvent& PropertyChangedEvent)
 {
 	const FName PropertyName = PropertyChangedEvent.GetPropertyName();
@@ -243,6 +535,13 @@ void SCineAssemblySchemaWindow::OnSchemaWindowClosed(const TSharedRef<SWindow>&
 	{
 		SchemaToConfigure->RenameAsset();
 	}
+
+	// Reset the schema's subsequence list to the current state of the asset list
+	SchemaToConfigure->Subsequences.Empty();
+	for (const TSharedPtr<FString>& AssetItem : AssetList)
+	{
+		SchemaToConfigure->Subsequences.Add(*AssetItem);
+	}
 }
 
 #undef LOCTEXT_NAMESPACE
diff --git a/Engine/Plugins/VirtualProduction/Shots/Source/ShotsEditor/Private/UI/CineAssembly/SCineAssemblySchemaWindow.h b/Engine/Plugins/VirtualProduction/Shots/Source/ShotsEditor/Private/UI/CineAssembly/SCineAssemblySchemaWindow.h
--- a/Engine/Plugins/VirtualProduction/Shots/Source/ShotsEditor/Private/UI/CineAssembly/SCineAssemblySchemaWindow.h
+++ b/Engine/Plugins/VirtualProduction/Shots/Source/ShotsEditor/Private/UI/CineAssembly/SCineAssemblySchemaWindow.h
@@ -6,6 +6,8 @@
 
 #include "CineAssemblySchema.h"
 #include "Widgets/Layout/SWidgetSwitcher.h"
+#include "Widgets/Text/SInlineEditableTextBlock.h"
+#include "Widgets/Views/SListView.h"
 
 /** Mode that configures the UI based on the intended user interactions with the Cine Assembly Schema asset */
 enum class ESchemaConfigMode : uint8
@@ -14,6 +16,14 @@ enum class ESchemaConfigMode : uint8
 	Edit
 };
 
+/** Row widget for the asset list view */
+class SSchemaAssetTableRow : public STableRow<TSharedPtr<FString>>
+{
+public:
+	/** Editable text widget stored in this row to easily trigger edit mode */
+	TSharedPtr<SInlineEditableTextBlock> TextBlock;
+};
+
 /**
  * A window for configuring the properties of a UCineAssemblySchema
  */
@@ -43,18 +53,39 @@ private:
 	/** Creates the content for the Assets tab */
 	TSharedRef<SWidget> MakeAssetTabContent();
 
+	/** Creates the content for the Hierarchy tab */
+	TSharedRef<SWidget> MakeHierarchyTabContent();
+
 	/** Closes the window and indicates that a new asset should be created by the asset factory */
 	FReply OnCreateAssetClicked();
 
 	/** Closes the window and indicates that no assets should be created by the asset factory */
 	FReply OnCancelClicked();
 
+	/** Adds a new entry to the asset list, which will appear in the list view, and will be added to the schema when the window closes */
+	FReply OnAddAssetClicked();
+
+	/** Generates a row in the asset list view */
+	TSharedRef<ITableRow> OnGenerateAssetRow(TSharedPtr<FString> InItem, const TSharedRef<STableViewBase>& OwnerTable);
+
+	/** Callback when the asset list view rebuilds itself, used to make the textbox in the most recently added row editable */
+	void OnAssetListRebuilt();
+
+	/** Create the context menu when the asset list view is right-clicked */
+	TSharedPtr<SWidget> MakeAssetListContextMenu();
+
 	/** Callback when one of the properties of the schema being configured changes */
 	void OnSchemaPropertiesChanged(const FPropertyChangedEvent& PropertyChangedEvent);
 
 	/** Callback when this window is about to be closed */
 	void OnSchemaWindowClosed(const TSharedRef<SWindow>& InWindow);
 
+	/** Finds the row for the input item and puts its textblock into edit mode */
+	void StartRename(TSharedPtr<FString> ItemToRename);
+
+	/** Validate the text entered by the user to ensure it will be a valid asset name */
+	bool ValidateAssetName(const FText& InText, FText& OutErrorMessage);
+
 private:
 	/** Switcher that controls which tab widget is currently visible */
 	TSharedPtr<SWidgetSwitcher> TabSwitcher;
@@ -70,4 +101,10 @@ private:
 
 	/** Whether the schema asset being configured needs to be renamed */
 	bool bPendingRename = false;
+
+	/** List items sources for the asset list view */
+	TArray<TSharedPtr<FString>> AssetList;
+
+	/** List view of assets that will be associate with the schema */
+	TSharedPtr<SListView<TSharedPtr<FString>>> AssetListView;
 };
