Shot Management: Add feature to Schemas to assign the generated subsequences to specific folder locations.
Also adds a feature to Schemas to allow for empty folders to be created when creating assemblies from that schema.
#rnx

diff --git a/Engine/Plugins/VirtualProduction/Shots/Source/ShotsEditor/Private/CineAssemblyFactory.cpp b/Engine/Plugins/VirtualProduction/Shots/Source/ShotsEditor/Private/CineAssemblyFactory.cpp
--- a/Engine/Plugins/VirtualProduction/Shots/Source/ShotsEditor/Private/CineAssemblyFactory.cpp
+++ b/Engine/Plugins/VirtualProduction/Shots/Source/ShotsEditor/Private/CineAssemblyFactory.cpp
@@ -2,8 +2,10 @@
 
 #include "CineAssemblyFactory.h"
 
+#include "AssetRegistry/AssetRegistryModule.h"
 #include "AssetToolsModule.h"
 #include "Framework/Application/SlateApplication.h"
+#include "HAL/FileManager.h"
 #include "Interfaces/IMainFrameModule.h"
 #include "Modules/ModuleManager.h"
 #include "MovieScene.h"
@@ -51,10 +53,33 @@ UObject* UCineAssemblyFactory::FactoryCreateNew(UClass* Class, UObject* InParent
 	FFrameRate TickResolution = CineAssemblyToConfigure->GetMovieScene()->GetTickResolution();
 	CineAssemblyToConfigure->GetMovieScene()->SetPlaybackRange((MovieSceneToolsProjectSettings->DefaultStartTime * TickResolution).FloorToFrame(), (MovieSceneToolsProjectSettings->DefaultDuration * TickResolution).FloorToFrame().Value);
 
-	// Create the subsequence assets associated with this assembly
 	if (const UCineAssemblySchema* Schema = CineAssemblyToConfigure->GetSchema())
 	{
-		if (Schema->Subsequences.Num() > 0)
+		// Get the path where the top-level assembly will be created so we can create other assets relative to it
+		const FAssetData AssemblyAssetData = FAssetData(CineAssemblyToConfigure);
+		const FString PackagePath = AssemblyAssetData.PackagePath.ToString();
+
+		// Create the default folders for this assembly, based on the schema
+		for (const FString& FolderPath : Schema->FoldersToCreate)
+		{
+			const FString PathToCreate = PackagePath / FolderPath;
+			const FString RelativeFilePath = FPackageName::LongPackageNameToFilename(PathToCreate);
+			const FString AbsoluteFilePath = FPaths::ConvertRelativePathToFull(RelativeFilePath);
+
+			// Create the directory on disk, then add its path to the asset registry so it appears in Content Browser
+			if (!IFileManager::Get().DirectoryExists(*AbsoluteFilePath))
+			{
+				constexpr bool bCreateParentFoldersIfMissing = true;
+				if (IFileManager::Get().MakeDirectory(*AbsoluteFilePath, bCreateParentFoldersIfMissing))
+				{
+					FAssetRegistryModule& AssetRegistryModule = FModuleManager::GetModuleChecked<FAssetRegistryModule>("AssetRegistry");
+					AssetRegistryModule.Get().AddPath(PathToCreate);
+				}
+			}
+		}
+
+		// Create the subsequence assets associated with this assembly
+		if (Schema->SubsequencesToCreate.Num() > 0)
 		{
 			// Add a subsequence track to the assembly's sequence
 			UMovieScene* MovieScene = CineAssemblyToConfigure->GetMovieScene();
@@ -62,20 +87,16 @@ UObject* UCineAssemblyFactory::FactoryCreateNew(UClass* Class, UObject* InParent
 
 			IAssetTools& AssetTools = FModuleManager::GetModuleChecked<FAssetToolsModule>("AssetTools").Get();
 
-			// Get the path where the top-level assembly will be created so we can create the subsequences in the same location
-			// TODO: Once schemas support specifying the path for subsequences, this code should use those paths relative to this package path
-			const FAssetData AssemblyAssetData = FAssetData(CineAssemblyToConfigure);
-			const FString PackagePath = AssemblyAssetData.PackagePath.ToString();
-
 			// Create a new CineAssembly for each subsequence, set its playback range to match the parent sequence, and add it to the subsequence track
-			for (const FString& Subsequence : Schema->Subsequences)
+			for (const FString& Subsequence : Schema->SubsequencesToCreate)
 			{
 				// Before creating each subassembly, sanity check that each one will actually have a unique name (in case there are duplicates in the schema description)
 				FString UniquePackageName;
 				FString UniqueAssetName;
 				AssetTools.CreateUniqueAssetName(PackagePath / Subsequence, TEXT(""), UniquePackageName, UniqueAssetName);
 
-				UCineAssembly* SubAssembly = Cast<UCineAssembly>(AssetTools.CreateAsset(UniqueAssetName, PackagePath, UCineAssembly::StaticClass(), NewObject<UCineAssemblyFactory>()));
+				const FString SubAssemblyPath = FPaths::GetPath(UniquePackageName);
+				UCineAssembly* SubAssembly = Cast<UCineAssembly>(AssetTools.CreateAsset(UniqueAssetName, SubAssemblyPath, UCineAssembly::StaticClass(), NewObject<UCineAssemblyFactory>()));
 
 				if (SubAssembly)
 				{
diff --git a/Engine/Plugins/VirtualProduction/Shots/Source/ShotsEditor/Private/CineAssemblySchema.h b/Engine/Plugins/VirtualProduction/Shots/Source/ShotsEditor/Private/CineAssemblySchema.h
--- a/Engine/Plugins/VirtualProduction/Shots/Source/ShotsEditor/Private/CineAssemblySchema.h
+++ b/Engine/Plugins/VirtualProduction/Shots/Source/ShotsEditor/Private/CineAssemblySchema.h
@@ -49,9 +49,13 @@ public:
 	UPROPERTY(EditAnywhere, Category = "Default")
 	TObjectPtr<UTexture2D> ThumbnailImage;
 
-	/** The names of the subsequences that should be created for assemblies that use this schema */
+	/** Paths of subsequence assets that should be created for assemblies that use this schema, relative to the path of the top-level assembly */
 	UPROPERTY()
-	TArray<FString> Subsequences;
+	TArray<FString> SubsequencesToCreate;
+
+	/** Paths of folders that should be created for assemblies that use this schema, relative to the path of the top-level assembly */
+	UPROPERTY()
+	TArray<FString> FoldersToCreate;
 
 private:
 	/** Slate brush used to actually render the thumbnail image */
diff --git a/Engine/Plugins/VirtualProduction/Shots/Source/ShotsEditor/Private/ShotsEditorStyle.cpp b/Engine/Plugins/VirtualProduction/Shots/Source/ShotsEditor/Private/ShotsEditorStyle.cpp
--- a/Engine/Plugins/VirtualProduction/Shots/Source/ShotsEditor/Private/ShotsEditorStyle.cpp
+++ b/Engine/Plugins/VirtualProduction/Shots/Source/ShotsEditor/Private/ShotsEditorStyle.cpp
@@ -38,6 +38,7 @@ FShotsEditorStyle::FShotsEditorStyle()
 	Set("Icons.Animation", new FSlateVectorImageBrush(RootToContentDir(TEXT("Slate/Animation_16.svg")), Icon16x16));
 	Set("Icons.Notes", new FSlateVectorImageBrush(RootToContentDir(TEXT("Slate/Note.svg")), Icon16x16));
 	Set("Icons.Shot", new FSlateVectorImageBrush(RootToContentDir(TEXT("Slate/Shot_16.svg")), Icon16x16));
+	Set("Icons.DragHandle", new FSlateVectorImageBrush(RootToContentDir(TEXT("Slate/drag-handle.svg")), Icon16x16));
 
 	Set("Badges.RevisionControlWarning", new FSlateVectorImageBrush(RootToCoreContentDir(TEXT("Starship/SourceControl/Status/RevisionControlBadgeWarning.svg")), Icon16x16, FStyleColors::Warning));
 	Set("Badges.RevisionControlConnected", new FSlateVectorImageBrush(RootToCoreContentDir(TEXT("Starship/SourceControl/Status/RevisionControlBadgeConnected.svg")), Icon16x16, FStyleColors::Success));
diff --git a/Engine/Plugins/VirtualProduction/Shots/Source/ShotsEditor/Private/UI/CineAssembly/SCineAssemblySchemaWindow.cpp b/Engine/Plugins/VirtualProduction/Shots/Source/ShotsEditor/Private/UI/CineAssembly/SCineAssemblySchemaWindow.cpp
--- a/Engine/Plugins/VirtualProduction/Shots/Source/ShotsEditor/Private/UI/CineAssembly/SCineAssemblySchemaWindow.cpp
+++ b/Engine/Plugins/VirtualProduction/Shots/Source/ShotsEditor/Private/UI/CineAssembly/SCineAssemblySchemaWindow.cpp
@@ -2,6 +2,8 @@
 
 #include "SCineAssemblySchemaWindow.h"
 
+#include "AssetRegistry/AssetRegistryModule.h"
+#include "AssetViewUtils.h"
 #include "Framework/MultiBox/MultiBoxBuilder.h"
 #include "IStructureDetailsView.h"
 #include "Modules/ModuleManager.h"
@@ -18,6 +20,11 @@
 
 #define LOCTEXT_NAMESPACE "SCineAssemblySchemaWindow"
 
+namespace UE::CineAssemblyTools::Private
+{
+	auto SortTreeItems = [](const TSharedPtr<FSchemaTreeItem>& A, const TSharedPtr<FSchemaTreeItem>& B) { return A->Path < B->Path; };
+};
+
 void SCineAssemblySchemaWindow::Construct(const FArguments& InArgs, UCineAssemblySchema* InSchema, ESchemaConfigMode InConfigMode)
 {
 	SchemaToConfigure = InSchema;
@@ -35,13 +42,10 @@ void SCineAssemblySchemaWindow::Construct(const FArguments& InArgs, UCineAssembl
 		WindowTitle = LOCTEXT("WindowTitleEdit", "Edit Assembly Schema");
 	}
 
-	GetOnWindowClosedEvent().AddSP(this, &SCineAssemblySchemaWindow::OnSchemaWindowClosed);
+	// Initialize the content tree view with the current list of assets and folders saved in the schema
+	InitializeContentTree();
 
-	// Initialize the asset list with the current list saved in the schema
-	for (const FString& Subsequence : SchemaToConfigure->Subsequences)
-	{
-		AssetList.Add(MakeShared<FString>(Subsequence));
-	}
+	GetOnWindowClosedEvent().AddSP(this, &SCineAssemblySchemaWindow::OnSchemaWindowClosed);
 
 	SWindow::Construct(SWindow::FArguments()
 		.Title(WindowTitle)
@@ -116,11 +120,6 @@ TSharedRef<SWidget> SCineAssemblySchemaWindow::MakeMenuPanel()
 				.HAlign(HAlign_Left)
 
 				+ SSegmentedControl<int32>::Slot(1)
-				.Text(LOCTEXT("TimelineTab", "Timeline Template"))
-				.Icon(FSlateIcon(FShotsEditorStyle::StyleName, "Icons.Sequencer").GetIcon())
-				.HAlign(HAlign_Left)
-
-				+ SSegmentedControl<int32>::Slot(2)
 				.Text(LOCTEXT("HierarchyTab", "Hierarchy"))
 				.Icon(FSlateIcon(FAppStyle::GetAppStyleSetName(), "Icons.FolderClosed").GetIcon())
 				.HAlign(HAlign_Left)
@@ -129,31 +128,12 @@ TSharedRef<SWidget> SCineAssemblySchemaWindow::MakeMenuPanel()
 
 TSharedRef<SWidget> SCineAssemblySchemaWindow::MakeContentPanel()
 {
-	// TODO: Remove this when all of the tab content has been implemented
-	auto BuildStubWidget = [](const FString& StubText) -> TSharedRef<SWidget>
-		{
-			return SNew(SHorizontalBox)
-
-			+ SHorizontalBox::Slot()
-				.FillWidth(1.0f)
-				.VAlign(VAlign_Center)
-				.HAlign(HAlign_Center)
-				[
-					SNew(STextBlock).Text(FText::FromString(StubText))
-				];
-		};
-
 	TabSwitcher = SNew(SWidgetSwitcher)
 		+ SWidgetSwitcher::Slot()
 		[
 			MakeAssetTabContent()
 		]
 
-		+ SWidgetSwitcher::Slot()
-		[
-			BuildStubWidget("TIMELINE TAB CONTENT")
-		]
-
 		+ SWidgetSwitcher::Slot()
 		[
 			MakeHierarchyTabContent()
@@ -220,6 +200,17 @@ TSharedRef<SWidget> SCineAssemblySchemaWindow::MakeAssetTabContent()
 
 TSharedRef<SWidget> SCineAssemblySchemaWindow::MakeHierarchyTabContent()
 {
+	TreeView = SNew(STreeView<TSharedPtr<FSchemaTreeItem>>)
+		.TreeItemsSource(&TreeItems)
+		.SelectionMode(ESelectionMode::Single)
+		.OnGenerateRow(this, &SCineAssemblySchemaWindow::OnGenerateTreeRow)
+		.OnGetChildren(this, &SCineAssemblySchemaWindow::OnGetChildren)
+		.OnItemsRebuilt(this, &SCineAssemblySchemaWindow::OnTreeItemsRebuilt)
+		.OnContextMenuOpening(this, &SCineAssemblySchemaWindow::MakeContentTreeContextMenu)
+		.OnKeyDownHandler(this, &SCineAssemblySchemaWindow::OnTreeViewKeyDown);
+
+	ExpandTreeRecursive(RootItem);
+
 	TSharedRef<SWidget> ContentHierarhcyWidget = SNew(SVerticalBox)
 
 		+ SVerticalBox::Slot()
@@ -256,9 +247,9 @@ TSharedRef<SWidget> SCineAssemblySchemaWindow::MakeHierarchyTabContent()
 		.HAlign(HAlign_Left)
 		.Padding(0.0f, 0.0f, 0.0f, 8.0f)
 		[
-			// TODO: Implement this add button to build out an actual tree view to drag assets from the asset list onto
 			SNew(SButton)
 				.ContentPadding(FMargin(2.0f))
+				.OnClicked(this, &SCineAssemblySchemaWindow::OnAddFolderClicked)
 				[
 					SNew(SHorizontalBox)
 
@@ -274,16 +265,18 @@ TSharedRef<SWidget> SCineAssemblySchemaWindow::MakeHierarchyTabContent()
 					+ SHorizontalBox::Slot()
 						.AutoWidth()
 						[
-							SNew(STextBlock).Text(LOCTEXT("AddFolderButton", "Add"))
+							SNew(STextBlock).Text(LOCTEXT("AddFolderButton", "Add Folder"))
 						]
 		]		]
 
 		+ SVerticalBox::Slot()
 		.FillHeight(1.0f)
 		[
-			// TODO: Implement the tree view widget here for the folder hierarchy
 			SNew(SBorder)
 				.BorderImage(FShotsEditorStyle::Get().GetBrush("ProductionWizard.RecessedBackground"))
+				[
+					TreeView.ToSharedRef()
+				]
 		];
 
 	TSharedRef<SWidget> AssetListWidget = SNew(SVerticalBox)
@@ -321,7 +314,7 @@ TSharedRef<SWidget> SCineAssemblySchemaWindow::MakeHierarchyTabContent()
 					+ SHorizontalBox::Slot()
 						.AutoWidth()
 						[
-							SNew(STextBlock).Text(LOCTEXT("AddAssetButton", "Add"))
+							SNew(STextBlock).Text(LOCTEXT("AddAssetButton", "Add Asset"))
 						]
 		]		]
 
@@ -332,12 +325,13 @@ TSharedRef<SWidget> SCineAssemblySchemaWindow::MakeHierarchyTabContent()
 				.BorderImage(FShotsEditorStyle::Get().GetBrush("ProductionWizard.RecessedBackground"))
 				[
 					SAssignNew(AssetListView, SListView<TSharedPtr<FString>>)
-						.ListItemsSource(&AssetList)
+						.ListItemsSource(&AssetListItems)
 						.SelectionMode(ESelectionMode::Single)
 						.OnGenerateRow(this, &SCineAssemblySchemaWindow::OnGenerateAssetRow)
 						.OnItemsRebuilt(this, &SCineAssemblySchemaWindow::OnAssetListRebuilt)
 						.OnContextMenuOpening(this, &SCineAssemblySchemaWindow::MakeAssetListContextMenu)
-						.OnMouseButtonDoubleClick(this, &SCineAssemblySchemaWindow::StartRename)
+						.OnMouseButtonDoubleClick(this, &SCineAssemblySchemaWindow::EnterEditMode)
+						.OnKeyDownHandler(this, &SCineAssemblySchemaWindow::OnAssetListKeyDown)
 				]
 		];
 
@@ -347,7 +341,7 @@ TSharedRef<SWidget> SCineAssemblySchemaWindow::MakeHierarchyTabContent()
 		.FillWidth(0.5f)
 		.Padding(16.0f)
 		[
-			ContentHierarhcyWidget
+			AssetListWidget
 		]
 
 		+ SHorizontalBox::Slot()
@@ -362,23 +356,58 @@ TSharedRef<SWidget> SCineAssemblySchemaWindow::MakeHierarchyTabContent()
 		.FillWidth(0.5f)
 		.Padding(16.0f)
 		[
-			AssetListWidget
+			ContentHierarhcyWidget
 		];
 }
 
 FReply SCineAssemblySchemaWindow::OnAddAssetClicked()
 {
 	// Add a new empty string to the list, which will be renamed by the user
-	AssetList.Add(MakeShared<FString>(TEXT("")));
+	AssetListItems.Add(MakeShared<FString>(TEXT("")));
 	AssetListView->RequestListRefresh();
 	return FReply::Handled();
 }
 
+FReply SCineAssemblySchemaWindow::OnAddFolderClicked()
+{
+	TSharedPtr<FSchemaTreeItem> NewFolder = MakeShared<FSchemaTreeItem>();
+	NewFolder->Type = FSchemaTreeItem::EItemType::Folder;
+
+	// Get the parent item for the new folder being added (this can be the root folder if no parent is currently selected)
+	// The TreeView uses single selection mode, so at most one item can ever be selected by the user
+	TSharedPtr<FSchemaTreeItem> SelectedTreeItem = RootItem;
+
+	TArray<TSharedPtr<FSchemaTreeItem>> SelectedNodes = TreeView->GetSelectedItems();
+	if (SelectedNodes.Num() == 1)
+	{
+		SelectedTreeItem = SelectedNodes[0];
+	}
+
+	const FString ParentPath = SelectedTreeItem->Path;
+	const FString NewFolderName = TEXT("NewFolder");
+	NewFolder->Path = ParentPath / NewFolderName;
+	NewFolder->Parent = SelectedTreeItem;
+
+	SelectedTreeItem->ChildFolders.Add(NewFolder);
+
+	// Sort the children alphabetically to maintain a good ordering with the new folder
+	SelectedTreeItem->ChildFolders.Sort(UE::CineAssemblyTools::Private::SortTreeItems);
+
+	// Save a reference to this item so that when the tree is rebuilt, we can immediately start editing its name
+	MostRecentlyAddedItem = NewFolder;
+
+	TreeView->SetItemExpansion(SelectedTreeItem, true);
+	TreeView->RequestTreeRefresh();
+
+	return FReply::Handled();
+}
+
 TSharedRef<ITableRow> SCineAssemblySchemaWindow::OnGenerateAssetRow(TSharedPtr<FString> InItem, const TSharedRef<STableViewBase>& OwnerTable)
 {
 	TSharedRef<SSchemaAssetTableRow> Row = SNew(SSchemaAssetTableRow, OwnerTable)
 		.ShowSelection(true)
-		.Padding(FMargin(4.0f, 2.0f));
+		.Padding(FMargin(4.0f, 2.0f))
+		.OnDragDetected(this, &SCineAssemblySchemaWindow::OnAssetRowDragDetected);
 
 	// Store a reference to the editable textblock in the row to easily set it to edit mode for renaming
 	Row->TextBlock = SNew(SInlineEditableTextBlock)
@@ -412,6 +441,13 @@ TSharedRef<ITableRow> SCineAssemblySchemaWindow::OnGenerateAssetRow(TSharedPtr<F
 				[
 					Row->TextBlock.ToSharedRef()
 				]
+
+			+ SHorizontalBox::Slot()
+				.AutoWidth()
+				.VAlign(VAlign_Center)
+				[
+					SNew(SImage).Image(FShotsEditorStyle::Get().GetBrush("Icons.DragHandle"))
+				]
 		]);
 
 	return Row;
@@ -423,10 +459,10 @@ void SCineAssemblySchemaWindow::OnAssetListRebuilt()
 	{
 		// Only trigger edit mode on the last item in the asset list, and only if it is an empty string so that it can be renamed
 		TSharedPtr<FString> LastItem = AssetListView->GetItems().Last();
-		if (LastItem->IsEmpty())
+		if (LastItem && LastItem->IsEmpty())
 		{
 			*LastItem = TEXT("NewAsset");
-			StartRename(LastItem);
+			EnterEditMode(LastItem);
 		}
 	}
 }
@@ -446,11 +482,7 @@ TSharedPtr<SWidget> SCineAssemblySchemaWindow::MakeAssetListContextMenu()
 			LOCTEXT("RenameAsset", "Rename"),
 			LOCTEXT("RenameAssetToolTip", "Rename"),
 			FSlateIcon(FShotsEditorStyle::StyleName, "Icons.AssetNaming"),
-			FUIAction(FExecuteAction::CreateLambda([this, SelectedItem]()
-				{
-					StartRename(SelectedItem);
-				})
-			),
+			FUIAction(FExecuteAction::CreateSP(this, &SCineAssemblySchemaWindow::EnterEditMode, SelectedItem)),
 			NAME_None,
 			EUserInterfaceActionType::Button
 		);
@@ -459,16 +491,7 @@ TSharedPtr<SWidget> SCineAssemblySchemaWindow::MakeAssetListContextMenu()
 			LOCTEXT("DeleteAssetFromList", "Delete"),
 			LOCTEXT("DeleteAssetFromListToolTip", "Delete"),
 			FSlateIcon(FAppStyle::GetAppStyleSetName(), "Icons.Delete"),
-			FUIAction(FExecuteAction::CreateLambda([this, SelectedItem]()
-				{
-					const int32 IndexToRemove = AssetList.IndexOfByKey(SelectedItem);
-					if (AssetList.IsValidIndex(IndexToRemove))
-					{
-						AssetList.RemoveAt(IndexToRemove);
-					}
-					AssetListView->RequestListRefresh();
-				})
-			),
+			FUIAction(FExecuteAction::CreateSP(this, &SCineAssemblySchemaWindow::DeleteAssetItem, SelectedItem)),
 			NAME_None,
 			EUserInterfaceActionType::Button
 		);
@@ -479,16 +502,329 @@ TSharedPtr<SWidget> SCineAssemblySchemaWindow::MakeAssetListContextMenu()
 	return SNullWidget::NullWidget;
 }
 
-void SCineAssemblySchemaWindow::StartRename(TSharedPtr<FString> ItemToRename)
+FReply SCineAssemblySchemaWindow::OnAssetListKeyDown(const FGeometry& MyGeometry, const FKeyEvent& InKeyEvent)
+{
+	if (InKeyEvent.GetKey() == EKeys::Delete)
+	{
+		// The list view uses single selection mode, so at most one item can ever be selected by the user
+		TArray<TSharedPtr<FString>> SelectedItems = AssetListView->GetSelectedItems();
+		if (SelectedItems.Num() == 1)
+		{
+			const TSharedPtr<FString>& SelectedItem = SelectedItems[0];
+			DeleteAssetItem(SelectedItem);
+		}
+	}
+
+	return FReply::Handled();
+}
+
+TSharedRef<ITableRow> SCineAssemblySchemaWindow::OnGenerateTreeRow(TSharedPtr<FSchemaTreeItem> TreeItem, const TSharedRef<STableViewBase>& OwnerTable)
+{
+	TSharedPtr<SImage> Icon = SNew(SImage);
+
+	if (TreeItem->Type == FSchemaTreeItem::EItemType::Folder)
+	{
+		Icon->SetImage(FShotsEditorStyle::Get().GetBrush("Icons.Folder"));
+		Icon->SetColorAndOpacity(FAppStyle::Get().GetSlateColor("ContentBrowser.DefaultFolderColor"));
+	}
+	else
+	{
+		Icon->SetImage(FShotsEditorStyle::Get().GetBrush("Icons.Sequencer"));
+		Icon->SetColorAndOpacity(FLinearColor::White);
+	}
+
+	return SNew(STableRow<TSharedPtr<FSchemaTreeItem>>, OwnerTable)
+		.ShowSelection(true)
+		.Padding(FMargin(8.0f, 2.0f, 8.0f, 0.0f))
+		.OnCanAcceptDrop_Lambda([](const FDragDropEvent& InDragDropEvent, EItemDropZone InDropZone, TSharedPtr<FSchemaTreeItem> InItem)
+			{
+				// Only folder items can accept drops, and only onto to th item (not above or below)
+				return (InItem->Type == FSchemaTreeItem::EItemType::Folder) ? EItemDropZone::OntoItem : TOptional<EItemDropZone>();
+			})
+		.OnAcceptDrop(this, &SCineAssemblySchemaWindow::OnTreeRowAcceptDrop)
+		.OnDragDetected(this, &SCineAssemblySchemaWindow::OnTreeRowDragDetected)
+		.Content()
+		[
+			SNew(SHorizontalBox)
+
+			+ SHorizontalBox::Slot()
+				.AutoWidth()
+				.VAlign(VAlign_Center)
+				.Padding(0.0f, 0.0f, 2.0f, 0.0f)
+				[
+					Icon.ToSharedRef()
+				]
+
+			+ SHorizontalBox::Slot()
+				.VAlign(VAlign_Center)
+				[
+					SAssignNew(TreeItem->NameWidget, SInlineEditableTextBlock)
+						.Text_Lambda([this, TreeItem]() -> FText
+							{
+								return (TreeItem == RootItem) ? LOCTEXT("RootPathName", "Root Folder") : FText::FromString(FPaths::GetPathLeaf(TreeItem->Path));
+							})
+						.IsReadOnly_Lambda([this, TreeItem]() -> bool
+							{
+								return (TreeItem == RootItem);
+							})
+						.OnVerifyTextChanged_Lambda([this, TreeItem](const FText& InText, FText& OutErrorMessage) -> bool
+							{
+								if (TreeItem->Type == FSchemaTreeItem::EItemType::Folder)
+								{
+									return ValidateFolderName(InText, OutErrorMessage, TreeItem);
+								}
+								return ValidateAssetName(InText, OutErrorMessage);
+							})
+						.OnTextCommitted(this, &SCineAssemblySchemaWindow::OnTreeItemTextCommitted, TreeItem)
+				]
+		];
+}
+
+void SCineAssemblySchemaWindow::OnGetChildren(TSharedPtr<FSchemaTreeItem> TreeItem, TArray<TSharedPtr<FSchemaTreeItem>>& OutNodes)
+{
+	// Display all of the child assets first, followed by all of the child folders
+	OutNodes.Append(TreeItem->ChildAssets);
+	OutNodes.Append(TreeItem->ChildFolders);
+}
+
+void SCineAssemblySchemaWindow::OnTreeItemsRebuilt()
+{
+	// Upon regenerating the tree view, allow the user to immediately interact with the name widget of the newly added folder in order to rename it
+	if (MostRecentlyAddedItem && MostRecentlyAddedItem->NameWidget)
+	{
+		FSlateApplication::Get().SetKeyboardFocus(MostRecentlyAddedItem->NameWidget.ToSharedRef());
+		MostRecentlyAddedItem->NameWidget->EnterEditingMode();
+		MostRecentlyAddedItem.Reset();
+	}
+}
+
+TSharedPtr<SWidget> SCineAssemblySchemaWindow::MakeContentTreeContextMenu()
+{
+	// The TreeView uses single selection mode, so at most one item can ever be selected by the user
+	TArray<TSharedPtr<FSchemaTreeItem>> SelectedNodes = TreeView->GetSelectedItems();
+	if (SelectedNodes.Num() == 1)
+	{
+		constexpr bool bCloseAfterMenuSelection = true;
+		FMenuBuilder MenuBuilder(bCloseAfterMenuSelection, nullptr);
+
+		TSharedPtr<FSchemaTreeItem>& SelectedTreeItem = SelectedNodes[0];
+
+		if (SelectedTreeItem->Type == FSchemaTreeItem::EItemType::Folder)
+		{
+			MenuBuilder.AddMenuEntry(
+				LOCTEXT("AddChildFolderAction", "Add Child Folder"),
+				LOCTEXT("AddChildFolderTooltip", "Add Child Folder"),
+				FSlateIcon(FAppStyle::GetAppStyleSetName(), "Icons.Plus"),
+				FUIAction(FExecuteAction::CreateLambda([this]() { OnAddFolderClicked(); })),
+				NAME_None,
+				EUserInterfaceActionType::Button
+			);
+		}
+
+		if (SelectedTreeItem != RootItem)
+		{
+			MenuBuilder.AddMenuEntry(
+				LOCTEXT("RenameAction", "Rename"),
+				LOCTEXT("RenameActionToolTip", "Rename"),
+				FSlateIcon(FShotsEditorStyle::StyleName, "Icons.AssetNaming"),
+				FUIAction(
+					FExecuteAction::CreateLambda([SelectedTreeItem]()
+						{
+							FSlateApplication::Get().SetKeyboardFocus(SelectedTreeItem->NameWidget.ToSharedRef());
+							SelectedTreeItem->NameWidget->EnterEditingMode();
+						})),
+				NAME_None,
+				EUserInterfaceActionType::Button
+			);
+
+			MenuBuilder.AddMenuEntry(
+				LOCTEXT("DeleteAction", "Delete"),
+				LOCTEXT("DeleteActionToolTip", "Delete"),
+				FSlateIcon(FAppStyle::GetAppStyleSetName(), "Icons.Delete"),
+				FUIAction(FExecuteAction::CreateSP(this, &SCineAssemblySchemaWindow::DeleteTreeItem, SelectedTreeItem)),
+				NAME_None,
+				EUserInterfaceActionType::Button
+			);
+		}
+
+		return MenuBuilder.MakeWidget();
+	}
+
+	return SNullWidget::NullWidget;
+}
+
+FReply SCineAssemblySchemaWindow::OnTreeViewKeyDown(const FGeometry& MyGeometry, const FKeyEvent& InKeyEvent)
+{
+	if (InKeyEvent.GetKey() == EKeys::Delete)
+	{
+		// The TreeView uses single selection mode, so at most one item can ever be selected by the user
+		TArray<TSharedPtr<FSchemaTreeItem>> SelectedItems = TreeView->GetSelectedItems();
+		if (SelectedItems.Num() == 1)
+		{
+			const TSharedPtr<FSchemaTreeItem>& SelectedItem = SelectedItems[0];
+			DeleteTreeItem(SelectedItem);
+		}
+	}
+
+	return FReply::Handled();
+}
+
+void SCineAssemblySchemaWindow::InitializeContentTree()
+{
+	// Create the content tree root
+	RootItem = MakeShared<FSchemaTreeItem>();
+	RootItem->Type = FSchemaTreeItem::EItemType::Folder;
+	RootItem->Path = TEXT("");
+
+	TreeItems.Add(RootItem);
+
+	auto AddItemsToTree = [this](TArray<FString>& ItemList, FSchemaTreeItem::EItemType ItemType)
+		{
+			// Sort the list so that paths are added to the tree in the proper order
+			ItemList.Sort();
+
+			for (const FString& ItemName : ItemList)
+			{
+				const FString ParentPath = FPaths::GetPath(ItemName);
+
+				// Walk the tree until we find an item whose path matches the parent path. The new tree item will be created as one of its children
+				TSharedPtr<FSchemaTreeItem> ParentItem = FindItemAtPathRecursive(RootItem, ParentPath);
+				if (ParentItem)
+				{
+					TSharedPtr<FSchemaTreeItem> NewItem = MakeShared<FSchemaTreeItem>();
+					NewItem->Type = ItemType;
+					NewItem->Path = ItemName;
+					NewItem->Parent = ParentItem;
+
+					if (ItemType == FSchemaTreeItem::EItemType::Folder)
+					{
+						ParentItem->ChildFolders.Add(NewItem);
+					}
+					else
+					{
+						ParentItem->ChildAssets.Add(NewItem);
+					}
+				}
+			}
+		};
+
+	AddItemsToTree(SchemaToConfigure->FoldersToCreate, FSchemaTreeItem::EItemType::Folder);
+	AddItemsToTree(SchemaToConfigure->SubsequencesToCreate, FSchemaTreeItem::EItemType::Asset);
+}
+
+void SCineAssemblySchemaWindow::ExpandTreeRecursive(TSharedPtr<FSchemaTreeItem> TreeItem) const
+{
+	TreeView->SetItemExpansion(TreeItem, true);
+
+	for (const TSharedPtr<FSchemaTreeItem>& ChildItem : TreeItem->ChildFolders)
+	{
+		ExpandTreeRecursive(ChildItem);
+	}
+}
+
+void SCineAssemblySchemaWindow::GetFolderListRecursive(TSharedPtr<FSchemaTreeItem> TreeItem, TArray<FString>& FolderList) const
+{
+	for (const TSharedPtr<FSchemaTreeItem>& Child : TreeItem->ChildFolders)
+	{
+		FolderList.Add(Child->Path);
+		GetFolderListRecursive(Child, FolderList);
+	}
+}
+
+void SCineAssemblySchemaWindow::GetAssetListRecursive(TSharedPtr<FSchemaTreeItem> TreeItem, TArray<FString>& AssetPathList) const
+{
+	for (const TSharedPtr<FSchemaTreeItem>& Asset : TreeItem->ChildAssets)
+	{
+		AssetPathList.Add(Asset->Path);
+	}
+
+	for (const TSharedPtr<FSchemaTreeItem>& Child : TreeItem->ChildFolders)
+	{
+		GetAssetListRecursive(Child, AssetPathList);
+	}
+}
+
+TSharedPtr<FSchemaTreeItem> SCineAssemblySchemaWindow::FindItemAtPathRecursive(TSharedPtr<FSchemaTreeItem> TreeItem, const FString& Path) const
+{
+	if (TreeItem->Path.Equals(Path))
+	{
+		return TreeItem;
+	}
+
+	for (const TSharedPtr<FSchemaTreeItem>& Child : TreeItem->ChildFolders)
+	{
+		if (const TSharedPtr<FSchemaTreeItem>& ItemAtPath = FindItemAtPathRecursive(Child, Path))
+		{
+			return ItemAtPath;
+		}
+	}
+
+	return nullptr;
+}
+
+void SCineAssemblySchemaWindow::SetChildrenPathRecursive(TSharedPtr<FSchemaTreeItem> TreeItem, const FString& NewPath)
+{
+	for (TSharedPtr<FSchemaTreeItem>& Asset : TreeItem->ChildAssets)
+	{
+		const FString OldAssetName = FPaths::GetPathLeaf(Asset->Path);
+
+		const FString NewAssetPath = NewPath / OldAssetName;
+		Asset->Path = NewAssetPath;
+	}
+
+	for (TSharedPtr<FSchemaTreeItem>& Child : TreeItem->ChildFolders)
+	{
+		const FString OldChildFolderName = FPaths::GetPathLeaf(Child->Path);
+
+		const FString NewChildPath = NewPath / OldChildFolderName;
+		Child->Path = NewChildPath;
+
+		SetChildrenPathRecursive(Child, NewChildPath);
+	}
+}
+
+void SCineAssemblySchemaWindow::EnterEditMode(TSharedPtr<FString> ItemToRename)
 {
 	// Get the row for the input item and put its textblock in edit mode so the user can rename the item
-	TSharedPtr<SSchemaAssetTableRow> Widget = StaticCastSharedPtr<SSchemaAssetTableRow>(AssetListView->WidgetFromItem(ItemToRename));
+	if (TSharedPtr<SSchemaAssetTableRow> Widget = StaticCastSharedPtr<SSchemaAssetTableRow>(AssetListView->WidgetFromItem(ItemToRename)))
+	{
+		if (Widget->TextBlock)
+		{
+			FSlateApplication::Get().SetKeyboardFocus(Widget->TextBlock.ToSharedRef());
+			Widget->TextBlock->EnterEditingMode();
+		}
+	}
+}
+
+void SCineAssemblySchemaWindow::OnTreeItemTextCommitted(const FText& InText, ETextCommit::Type InCommitType, TSharedPtr<FSchemaTreeItem> TreeItem)
+{
+	// Early-out if the name has not actually changed
+	const FString OldPath = TreeItem->Path;
+	const FString OldName = FPaths::GetPathLeaf(OldPath);
+	if (OldName.Equals(InText.ToString()))
+	{
+		return;
+	}
+
+	const FString NewPath = FPaths::GetPath(OldPath) / InText.ToString();
+	TreeItem->Path = NewPath;
+
+	if (TreeItem->Type == FSchemaTreeItem::EItemType::Folder)
+	{
+		// If this is a folder item, update the path of all of its children (recursively)
+		SetChildrenPathRecursive(TreeItem, NewPath);
+
+		TreeItem->Parent->ChildFolders.Sort(UE::CineAssemblyTools::Private::SortTreeItems);
+	}
+	else
+	{
+		TreeItem->Parent->ChildAssets.Sort(UE::CineAssemblyTools::Private::SortTreeItems);
+	}
 
-	FSlateApplication::Get().SetKeyboardFocus(Widget->TextBlock.ToSharedRef());
-	Widget->TextBlock->EnterEditingMode();
+	TreeView->RequestTreeRefresh();
 }
 
-bool SCineAssemblySchemaWindow::ValidateAssetName(const FText& InText, FText& OutErrorMessage)
+bool SCineAssemblySchemaWindow::ValidateAssetName(const FText& InText, FText& OutErrorMessage) const
 {
 	// An empty name is invalid
 	if (InText.IsEmpty())
@@ -502,6 +838,57 @@ bool SCineAssemblySchemaWindow::ValidateAssetName(const FText& InText, FText& Ou
 	return FName::IsValidXName(InText.ToString(), INVALID_OBJECTNAME_CHARACTERS INVALID_LONGPACKAGE_CHARACTERS, &OutErrorMessage);
 }
 
+bool SCineAssemblySchemaWindow::ValidateFolderName(const FText& InText, FText& OutErrorMessage, TSharedPtr<FSchemaTreeItem> TreeItem) const
+{
+	// Check for empty text string
+	if (InText.IsEmpty())
+	{
+		OutErrorMessage = LOCTEXT("EmptyNameErrorMessage", "Please provide a name for this folder");
+		return false;
+	}
+
+	// Check for invalid characters
+	if (!AssetViewUtils::IsValidFolderName(InText.ToString(), OutErrorMessage))
+	{
+		return false;
+	}
+
+	// Check for duplicate folder names
+	const FString ParentPath = FPaths::GetPath(TreeItem->Path);
+	TSharedPtr<FSchemaTreeItem> FoundItem = FindItemAtPathRecursive(RootItem, ParentPath / InText.ToString());
+
+	if (FoundItem && FoundItem != TreeItem)
+	{
+		OutErrorMessage = LOCTEXT("DuplicateNameErrorMessage", "A folder already exists at this location with this name");
+		return false;
+	}
+
+	return true;
+}
+
+void SCineAssemblySchemaWindow::DeleteAssetItem(TSharedPtr<FString> InItem)
+{
+	const int32 IndexToRemove = AssetListItems.IndexOfByKey(InItem);
+	if (AssetListItems.IsValidIndex(IndexToRemove))
+	{
+		AssetListItems.RemoveAt(IndexToRemove);
+	}
+	AssetListView->RequestListRefresh();
+}
+
+void SCineAssemblySchemaWindow::DeleteTreeItem(TSharedPtr<FSchemaTreeItem> TreeItem)
+{
+	if (TreeItem->Type == FSchemaTreeItem::EItemType::Folder)
+	{
+		TreeItem->Parent->ChildFolders.Remove(TreeItem);
+	}
+	else
+	{
+		TreeItem->Parent->ChildAssets.Remove(TreeItem);
+	}
+	TreeView->RequestTreeRefresh();
+}
+
 void SCineAssemblySchemaWindow::OnSchemaPropertiesChanged(const FPropertyChangedEvent& PropertyChangedEvent)
 {
 	const FName PropertyName = PropertyChangedEvent.GetPropertyName();
@@ -536,11 +923,194 @@ void SCineAssemblySchemaWindow::OnSchemaWindowClosed(const TSharedRef<SWindow>&
 		SchemaToConfigure->RenameAsset();
 	}
 
-	// Reset the schema's subsequence list to the current state of the asset list
-	SchemaToConfigure->Subsequences.Empty();
-	for (const TSharedPtr<FString>& AssetItem : AssetList)
+	// Update the schema's folder list
+	const TArray<FString> CachedFolderList = SchemaToConfigure->FoldersToCreate;
+
+	TArray<FString> FolderList;
+	GetFolderListRecursive(RootItem, FolderList);
+
+	if (FolderList != CachedFolderList)
+	{
+		SchemaToConfigure->Modify();
+		SchemaToConfigure->FoldersToCreate = FolderList;
+	}
+
+	// Update the schema's asset list
+	const TArray<FString> CachedAssetList = SchemaToConfigure->SubsequencesToCreate;
+
+	TArray<FString> AssetList;
+	GetAssetListRecursive(RootItem, AssetList);
+
+	if (AssetList != CachedAssetList)
+	{
+		SchemaToConfigure->Modify();
+		SchemaToConfigure->SubsequencesToCreate = AssetList;
+	}
+}
+
+FReply SCineAssemblySchemaWindow::OnAssetRowDragDetected(const FGeometry& MyGeometry, const FPointerEvent& MouseEvent)
+{
+	if (MouseEvent.IsMouseButtonDown(EKeys::LeftMouseButton))
+	{
+		// The ListView uses single selection mode, so at most one item can ever be selected by the user
+		TArray<TSharedPtr<FString>> SelectedItems = AssetListView->GetSelectedItems();
+		if (SelectedItems.Num() == 1)
+		{
+			const TSharedPtr<FString>& SelectedItem = SelectedItems[0];
+			TSharedRef<FSchemaAssetDragDrop> Operation = FSchemaAssetDragDrop::New(*SelectedItem);
+
+			Operation->OnDropNotHandled.BindLambda([this](const FString& AssetName)
+				{
+					// If the asset item is not dropped in a valid place, restore it to the asset list
+					AssetListItems.Add(MakeShared<FString>(AssetName));
+					AssetListView->RequestListRefresh();
+				});
+
+			// Remove this item from the asset list while it is being dragged
+			AssetListItems.Remove(SelectedItems[0]);
+			AssetListView->RequestListRefresh();
+
+			return FReply::Handled().BeginDragDrop(Operation);
+		}
+	}
+
+	return FReply::Unhandled();
+}
+
+FReply SCineAssemblySchemaWindow::OnTreeRowDragDetected(const FGeometry& MyGeometry, const FPointerEvent& MouseEvent)
+{
+	if (MouseEvent.IsMouseButtonDown(EKeys::LeftMouseButton))
+	{
+		TArray<TSharedPtr<FSchemaTreeItem>> SelectedTreeItems = TreeView->GetSelectedItems();
+		if (SelectedTreeItems.Num() == 1)
+		{
+			TSharedPtr<FSchemaTreeItem> SelectedTreeItem = SelectedTreeItems[0];
+			TSharedRef<FSchemaAssetDragDrop> Operation = FSchemaAssetDragDrop::New(FPaths::GetPathLeaf(SelectedTreeItem->Path));
+			Operation->SourceTreeItem = SelectedTreeItem;
+
+			if (SelectedTreeItem->Type == FSchemaTreeItem::EItemType::Asset)
+			{
+				Operation->OnDropNotHandled.BindLambda([this, SelectedTreeItem](const FString& AssetName)
+					{
+						// If the tree item is not dropped in a valid place, restore it to its original place in the tree
+						SelectedTreeItem->Parent->ChildAssets.Add(SelectedTreeItem);
+						TreeView->RequestTreeRefresh();
+					});
+
+				// Remove this item from the tree while it is being dragged
+				SelectedTreeItem->Parent->ChildAssets.Remove(SelectedTreeItem);
+			}
+			else
+			{
+				Operation->OnDropNotHandled.BindLambda([this, SelectedTreeItem](const FString& AssetName)
+					{
+						// If the tree item is not dropped in a valid place, restore it to its original place in the tree
+						SelectedTreeItem->Parent->ChildFolders.Add(SelectedTreeItem);
+						TreeView->RequestTreeRefresh();
+					});
+
+				// Remove this item from the tree while it is being dragged
+				SelectedTreeItem->Parent->ChildFolders.Remove(SelectedTreeItem);
+			}
+
+			TreeView->RequestTreeRefresh();
+			return FReply::Handled().BeginDragDrop(Operation);
+		}
+	}
+
+	return FReply::Unhandled();
+}
+
+FReply SCineAssemblySchemaWindow::OnTreeRowAcceptDrop(const FDragDropEvent& InDragDropEvent, EItemDropZone InDropZone, TSharedPtr<FSchemaTreeItem> InItem)
+{
+	TSharedPtr<FSchemaAssetDragDrop> Operation = InDragDropEvent.GetOperationAs<FSchemaAssetDragDrop>();
+	if (Operation.IsValid())
+	{
+		TSharedPtr<FSchemaTreeItem> NewItem;
+		if (Operation->SourceTreeItem)
+		{
+			NewItem = Operation->SourceTreeItem;
+		}
+		else
+		{
+			// This must be a new item from the asset list, so we need to make a new tree item for it
+			NewItem = MakeShared<FSchemaTreeItem>();
+			NewItem->Type = FSchemaTreeItem::EItemType::Asset;
+		}
+
+		const FString ParentPath = InItem->Path;
+		NewItem->Path = ParentPath / Operation->ItemName;
+
+		NewItem->Parent = InItem;
+
+		if (NewItem->Type == FSchemaTreeItem::EItemType::Asset)
+		{
+			InItem->ChildAssets.Add(NewItem);
+			InItem->ChildAssets.Sort(UE::CineAssemblyTools::Private::SortTreeItems);
+		}
+		else
+		{
+			InItem->ChildFolders.Add(NewItem);
+
+			SetChildrenPathRecursive(NewItem, NewItem->Path);
+			InItem->ChildFolders.Sort(UE::CineAssemblyTools::Private::SortTreeItems);
+		}
+
+		TreeView->SetItemExpansion(InItem, true);
+		TreeView->RequestTreeRefresh();
+	}
+
+	return FReply::Handled();
+}
+
+TSharedRef<FSchemaAssetDragDrop> FSchemaAssetDragDrop::New(const FString& InName)
+{
+	TSharedRef<FSchemaAssetDragDrop> DragDropOp = MakeShared<FSchemaAssetDragDrop>();
+	DragDropOp->ItemName = InName;
+	DragDropOp->MouseCursor = EMouseCursor::GrabHandClosed;
+	DragDropOp->Construct();
+
+	return DragDropOp;
+}
+
+TSharedPtr<SWidget> FSchemaAssetDragDrop::GetDefaultDecorator() const
+{
+	return SNew(SBorder)
+		.Padding(8.0f)
+		.BorderImage(FShotsEditorStyle::Get().GetBrush("ProductionWizard.PanelBackground"))
+		[
+			SNew(SHorizontalBox)
+
+			+ SHorizontalBox::Slot()
+				.AutoWidth()
+				.VAlign(VAlign_Center)
+				.Padding(0.0f, 0.0f, 8.0f, 0.0f)
+				[
+					SNew(SImage).Image(FShotsEditorStyle::Get().GetBrush("Icons.Sequencer"))
+				]
+
+			+ SHorizontalBox::Slot()
+				.FillWidth(1.0f)
+				.VAlign(VAlign_Center)
+				[
+					SNew(STextBlock).Text(FText::FromString(ItemName))
+				]
+		];
+}
+
+void FSchemaAssetDragDrop::OnDragged(const class FDragDropEvent& DragDropEvent)
+{
+	if (CursorDecoratorWindow.IsValid())
+	{
+		CursorDecoratorWindow->MoveWindowTo(DragDropEvent.GetScreenSpacePosition() - (CursorDecoratorWindow->GetSizeInScreen() * FVector2f(0.0f, 0.5f)));
+	}
+}
+
+void FSchemaAssetDragDrop::OnDrop(bool bDropWasHandled, const FPointerEvent& MouseEvent)
+{
+	if (!bDropWasHandled)
 	{
-		SchemaToConfigure->Subsequences.Add(*AssetItem);
+		OnDropNotHandled.ExecuteIfBound(ItemName);
 	}
 }
 
diff --git a/Engine/Plugins/VirtualProduction/Shots/Source/ShotsEditor/Private/UI/CineAssembly/SCineAssemblySchemaWindow.h b/Engine/Plugins/VirtualProduction/Shots/Source/ShotsEditor/Private/UI/CineAssembly/SCineAssemblySchemaWindow.h
--- a/Engine/Plugins/VirtualProduction/Shots/Source/ShotsEditor/Private/UI/CineAssembly/SCineAssemblySchemaWindow.h
+++ b/Engine/Plugins/VirtualProduction/Shots/Source/ShotsEditor/Private/UI/CineAssembly/SCineAssemblySchemaWindow.h
@@ -5,9 +5,11 @@
 #include "Widgets/SWindow.h"
 
 #include "CineAssemblySchema.h"
+#include "Input/DragAndDrop.h"
 #include "Widgets/Layout/SWidgetSwitcher.h"
 #include "Widgets/Text/SInlineEditableTextBlock.h"
 #include "Widgets/Views/SListView.h"
+#include "Widgets/Views/STreeView.h"
 
 /** Mode that configures the UI based on the intended user interactions with the Cine Assembly Schema asset */
 enum class ESchemaConfigMode : uint8
@@ -24,6 +26,64 @@ public:
 	TSharedPtr<SInlineEditableTextBlock> TextBlock;
 };
 
+/** An entry in the content folder tree view with knowledge of its path, children, and parent */
+struct FSchemaTreeItem
+{
+	/** The types of items that can be represented in this tree view */
+	enum class EItemType : uint8
+	{
+		Asset,
+		Folder
+	};
+
+	/** The type of this tree item */
+	EItemType Type;
+
+	/** The relative path of this tree item */
+	FString Path;
+
+	/** The parent of this item in the tree */
+	TSharedPtr<FSchemaTreeItem> Parent;
+
+	/** The children of this item in the tree that are Asset types */
+	TArray<TSharedPtr<FSchemaTreeItem>> ChildAssets;
+
+	/** The children of this item in the tree that are Folder types */
+	TArray<TSharedPtr<FSchemaTreeItem>> ChildFolders;
+
+	/** The text widget that displays this item's name and supports renaming */
+	TSharedPtr<SInlineEditableTextBlock> NameWidget;
+};
+
+/** Drag and Drop operation to handle dragging schema assets and folders  */
+class FSchemaAssetDragDrop : public FDragDropOperation
+{
+public:
+
+	DRAG_DROP_OPERATOR_TYPE(FSchemaAssetDragDrop, FDragDropOperation)
+
+	DECLARE_DELEGATE_OneParam(FOnDropNotHanlded, const FString&);
+
+	FSchemaAssetDragDrop() = default;
+
+	static TSharedRef<FSchemaAssetDragDrop> New(const FString& InName);
+
+	// Begin FDragDropOperation overrides
+	virtual void OnDrop(bool bDropWasHandled, const FPointerEvent& MouseEvent) override;
+	virtual void OnDragged(const class FDragDropEvent& DragDropEvent) override;
+	virtual TSharedPtr<SWidget> GetDefaultDecorator() const override;
+	// End FDragDropOperation overrides
+
+	/** The source item being dragged. This could be null if the thing being dragged is from the Asset List (not the tree) */
+	TSharedPtr<FSchemaTreeItem> SourceTreeItem;
+
+	/** The name of the item being displayed, used by the decorator and to name the new tree item that will get made when dropped */
+	FString ItemName;
+
+	/** Callback to handle the case where an item is dropped somewhere not on the tree view */
+	FOnDropNotHanlded OnDropNotHandled;
+};
+
 /**
  * A window for configuring the properties of a UCineAssemblySchema
  */
@@ -62,9 +122,12 @@ private:
 	/** Closes the window and indicates that no assets should be created by the asset factory */
 	FReply OnCancelClicked();
 
-	/** Adds a new entry to the asset list, which will appear in the list view, and will be added to the schema when the window closes */
+	/** Adds a new entry to the asset list view */
 	FReply OnAddAssetClicked();
 
+	/** Adds a new folder item to the content tree view */
+	FReply OnAddFolderClicked();
+
 	/** Generates a row in the asset list view */
 	TSharedRef<ITableRow> OnGenerateAssetRow(TSharedPtr<FString> InItem, const TSharedRef<STableViewBase>& OwnerTable);
 
@@ -74,17 +137,74 @@ private:
 	/** Create the context menu when the asset list view is right-clicked */
 	TSharedPtr<SWidget> MakeAssetListContextMenu();
 
+	/** Handles key presses on the asset list */
+	FReply OnAssetListKeyDown(const FGeometry& MyGeometry, const FKeyEvent& InKeyEvent);
+
+	/** Generates the row widget for an entry in the tree view */
+	TSharedRef<ITableRow> OnGenerateTreeRow(TSharedPtr<FSchemaTreeItem> TreeItem, const TSharedRef<STableViewBase>& OwnerTable);
+
+	/** Gets the children of the input tree view item to build additional tree rows */
+	void OnGetChildren(TSharedPtr<FSchemaTreeItem> TreeItem, TArray<TSharedPtr<FSchemaTreeItem>>& OutNodes);
+
+	/** Callback when the content tree view rebuilds itself, used to make the textbox in the most recently added row editable */
+	void OnTreeItemsRebuilt();
+
+	/** Create the context menu when the content tree view is right-clicked */
+	TSharedPtr<SWidget> MakeContentTreeContextMenu();
+
+	/** Handles key presses on the tree view */
+	FReply OnTreeViewKeyDown(const FGeometry& MyGeometry, const FKeyEvent& InKeyEvent);
+
+	/** Build the content tree from the list of folders and assets saved in the schema properties */
+	void InitializeContentTree();
+
+	/** Recursively expands every item in the tree view */
+	void ExpandTreeRecursive(TSharedPtr<FSchemaTreeItem> TreeItem) const;
+
+	/** Recursively get the path of every folder item in the tree view */
+	void GetFolderListRecursive(TSharedPtr<FSchemaTreeItem> TreeItem, TArray<FString>& FolderList) const;
+
+	/** Recursively get the path of every asset item in the tree view */
+	void GetAssetListRecursive(TSharedPtr<FSchemaTreeItem> TreeItem, TArray<FString>& AssetPathList) const;
+
+	/** Returns the tree item whose path matches the input path */
+	TSharedPtr<FSchemaTreeItem> FindItemAtPathRecursive(TSharedPtr<FSchemaTreeItem> TreeItem, const FString& Path) const;
+
+	/** Modifies the path of the input tree item and all of its children (recursively) with the input new path */
+	void SetChildrenPathRecursive(TSharedPtr<FSchemaTreeItem> TreeItem, const FString& NewPath);
+
+	/** Finds the row for the input item and puts its textblock into edit mode */
+	void EnterEditMode(TSharedPtr<FString> ItemToRename);
+
+	/** Renames the tree item, and updates the paths of all of its children */
+	void OnTreeItemTextCommitted(const FText& InText, ETextCommit::Type InCommitType, TSharedPtr<FSchemaTreeItem> TreeItem);
+
 	/** Callback when one of the properties of the schema being configured changes */
 	void OnSchemaPropertiesChanged(const FPropertyChangedEvent& PropertyChangedEvent);
 
 	/** Callback when this window is about to be closed */
 	void OnSchemaWindowClosed(const TSharedRef<SWindow>& InWindow);
 
-	/** Finds the row for the input item and puts its textblock into edit mode */
-	void StartRename(TSharedPtr<FString> ItemToRename);
-
 	/** Validate the text entered by the user to ensure it will be a valid asset name */
-	bool ValidateAssetName(const FText& InText, FText& OutErrorMessage);
+	bool ValidateAssetName(const FText& InText, FText& OutErrorMessage) const;
+
+	/** Validate the text entered by the user to ensure it will be a valid folder name */
+	bool ValidateFolderName(const FText& InText, FText& OutErrorMessage, TSharedPtr<FSchemaTreeItem> TreeItem) const;
+
+	/** Remove the input item from the asset list */
+	void DeleteAssetItem(TSharedPtr<FString> InItem);
+
+	/** Remove the input item and its children from the tree view */
+	void DeleteTreeItem(TSharedPtr<FSchemaTreeItem> TreeItem);
+
+	/** Begins a drag and drop event to drag an item out of the asset list view */
+	FReply OnAssetRowDragDetected(const FGeometry& MyGeometry, const FPointerEvent& MouseEvent);
+
+	/** Begins a drag and drop event to drag an item out of the content tree view */
+	FReply OnTreeRowDragDetected(const FGeometry& MyGeometry, const FPointerEvent& MouseEvent);
+
+	/** Callback when an item is dropped onto a folder in the content tree view to add it to the children of that row */
+	FReply OnTreeRowAcceptDrop(const FDragDropEvent& InDragDropEvent, EItemDropZone InDropZone, TSharedPtr<FSchemaTreeItem> InItem);
 
 private:
 	/** Switcher that controls which tab widget is currently visible */
@@ -103,8 +223,20 @@ private:
 	bool bPendingRename = false;
 
 	/** List items sources for the asset list view */
-	TArray<TSharedPtr<FString>> AssetList;
+	TArray<TSharedPtr<FString>> AssetListItems;
 
 	/** List view of assets that will be associate with the schema */
 	TSharedPtr<SListView<TSharedPtr<FString>>> AssetListView;
+
+	/** Items source for the tree view */
+	TArray<TSharedPtr<FSchemaTreeItem>> TreeItems;
+
+	/** The root item in the tree view */
+	TSharedPtr<FSchemaTreeItem> RootItem;
+
+	/** The tree view of content folders and assets for this schema */
+	TSharedPtr<STreeView<TSharedPtr<FSchemaTreeItem>>> TreeView;
+
+	/** The mostly recently added tree item, used to allow the user to immediately rename the item after the tree view updates to show it */
+	TSharedPtr<FSchemaTreeItem> MostRecentlyAddedItem;
 };
