Fix CIS Warning, duplicate loc key

#rnx

diff --git a/Engine/Plugins/VirtualProduction/Shots/Source/ShotsEditor/Private/UI/CineAssembly/SCineAssemblySchemaWindow.cpp b/Engine/Plugins/VirtualProduction/Shots/Source/ShotsEditor/Private/UI/CineAssembly/SCineAssemblySchemaWindow.cpp
--- a/Engine/Plugins/VirtualProduction/Shots/Source/ShotsEditor/Private/UI/CineAssembly/SCineAssemblySchemaWindow.cpp
+++ b/Engine/Plugins/VirtualProduction/Shots/Source/ShotsEditor/Private/UI/CineAssembly/SCineAssemblySchemaWindow.cpp
@@ -829,7 +829,7 @@ bool SCineAssemblySchemaWindow::ValidateAssetName(const FText& InText, FText& Ou
 	// An empty name is invalid
 	if (InText.IsEmpty())
 	{
-		OutErrorMessage = LOCTEXT("EmptyNameErrorMessage", "Please provide a name for the asset");
+		OutErrorMessage = LOCTEXT("EmptyAssetNameErrorMessage", "Please provide a name for the asset");
 		return false;
 	}
 
@@ -843,7 +843,7 @@ bool SCineAssemblySchemaWindow::ValidateFolderName(const FText& InText, FText& O
 	// Check for empty text string
 	if (InText.IsEmpty())
 	{
-		OutErrorMessage = LOCTEXT("EmptyNameErrorMessage", "Please provide a name for this folder");
+		OutErrorMessage = LOCTEXT("EmptyFolderNameErrorMessage", "Please provide a name for this folder");
 		return false;
 	}
 
