ShotManagement: Add additional BP/Python API for interacting with Production Settings

#rnx

diff --git a/Engine/Plugins/VirtualProduction/Shots/Source/ShotsEditor/Private/ProductionFunctionLibrary.cpp b/Engine/Plugins/VirtualProduction/Shots/Source/ShotsEditor/Private/ProductionFunctionLibrary.cpp
--- a/Engine/Plugins/VirtualProduction/Shots/Source/ShotsEditor/Private/ProductionFunctionLibrary.cpp
+++ b/Engine/Plugins/VirtualProduction/Shots/Source/ShotsEditor/Private/ProductionFunctionLibrary.cpp
@@ -8,8 +8,72 @@ const TArray<FCinematicProduction> UProductionFunctionLibrary::GetAllProductions
 	return ProductionSettings->GetProductions();
 }
 
+bool UProductionFunctionLibrary::GetProduction(FGuid ProductionID, FCinematicProduction& Production)
+{
+	const UProductionSettings* const ProductionSettings = GetDefault<UProductionSettings>();
+	TOptional<const FCinematicProduction> FoundProduction = ProductionSettings->GetProduction(ProductionID);
+	if (FoundProduction.IsSet())
+	{
+		Production = FoundProduction.GetValue();
+		return true;
+	}
+
+	Production.ProductionID.Invalidate();
+	return false;
+}
+
+bool UProductionFunctionLibrary::GetActiveProduction(FCinematicProduction& Production)
+{
+	const UProductionSettings* const ProductionSettings = GetDefault<UProductionSettings>();
+	TOptional<const FCinematicProduction> FoundProduction = ProductionSettings->GetActiveProduction();
+	if (FoundProduction.IsSet())
+	{
+		Production = FoundProduction.GetValue();
+		return true;
+	}
+
+	Production.ProductionID.Invalidate();
+	return false;
+}
+
 void UProductionFunctionLibrary::SetActiveProduction(FCinematicProduction Production)
 {
 	UProductionSettings* ProductionSettings = GetMutableDefault<UProductionSettings>();
 	ProductionSettings->SetActiveProduction(Production.ProductionID);
 }
+
+void UProductionFunctionLibrary::SetActiveProductionByID(FGuid ProductionID)
+{
+	UProductionSettings* ProductionSettings = GetMutableDefault<UProductionSettings>();
+	ProductionSettings->SetActiveProduction(ProductionID);
+}
+
+void UProductionFunctionLibrary::ClearActiveProduction()
+{
+	UProductionSettings* ProductionSettings = GetMutableDefault<UProductionSettings>();
+	ProductionSettings->SetActiveProduction(FGuid());
+}
+
+bool UProductionFunctionLibrary::IsActiveProduction(FGuid ProductionID)
+{
+	const UProductionSettings* const ProductionSettings = GetDefault<UProductionSettings>();
+	return ProductionSettings->IsActiveProduction(ProductionID);
+}
+
+void UProductionFunctionLibrary::AddProduction(FCinematicProduction Production)
+{
+	UProductionSettings* ProductionSettings = GetMutableDefault<UProductionSettings>();
+	ProductionSettings->AddProduction(Production);
+}
+
+void UProductionFunctionLibrary::DeleteProduction(FGuid ProductionID)
+{
+	UProductionSettings* ProductionSettings = GetMutableDefault<UProductionSettings>();
+	ProductionSettings->DeleteProduction(ProductionID);
+}
+
+void UProductionFunctionLibrary::RenameProduction(FGuid ProductionID, FString NewName)
+{
+	UProductionSettings* ProductionSettings = GetMutableDefault<UProductionSettings>();
+	ProductionSettings->RenameProduction(ProductionID, NewName);
+}
diff --git a/Engine/Plugins/VirtualProduction/Shots/Source/ShotsEditor/Private/ProductionFunctionLibrary.h b/Engine/Plugins/VirtualProduction/Shots/Source/ShotsEditor/Private/ProductionFunctionLibrary.h
--- a/Engine/Plugins/VirtualProduction/Shots/Source/ShotsEditor/Private/ProductionFunctionLibrary.h
+++ b/Engine/Plugins/VirtualProduction/Shots/Source/ShotsEditor/Private/ProductionFunctionLibrary.h
@@ -21,10 +21,42 @@ public:
 	UFUNCTION(BlueprintPure, Category = "Shot Management")
 	static const TArray<FCinematicProduction> GetAllProductions();
 
-	/** 
+	/** Get the Cinematic Production matching the input ProductionID (if it exists) */
+	UFUNCTION(BlueprintPure, Category = "Shot Management", meta=(ReturnDisplayName="IsValid"))
+	static bool GetProduction(FGuid ProductionID, FCinematicProduction& Production);
+
+	/** Get the active Cinematic Production (if it exists) */
+	UFUNCTION(BlueprintPure, Category = "Shot Management", meta=(ReturnDisplayName = "IsValid"))
+	static bool GetActiveProduction(FCinematicProduction& Production);
+
+	/**
 	 * Sets the input Production as the current Active Production 
 	 * If no input is provided, the Active Production will be set to None
 	 */
 	UFUNCTION(BlueprintCallable, Category = "Shot Management")
 	static void SetActiveProduction(FCinematicProduction Production = FCinematicProduction());
+
+	/** Sets the Production matching the input ProductionID as the current Active Production */
+	UFUNCTION(BlueprintCallable, Category = "Shot Management")
+	static void SetActiveProductionByID(FGuid ProductionID);
+
+	/** Sets the active Cinematic Production to None */
+	UFUNCTION(BlueprintCallable, Category = "Shot Management")
+	static void ClearActiveProduction();
+
+	/** Returns true if input ProductionID matches the ID of the current Active Production */
+	UFUNCTION(BlueprintPure, Category = "Shot Management")
+	static bool IsActiveProduction(FGuid ProductionID);
+
+	/** Add the input Cinematic Production to the Production Settings' list of productions */
+	UFUNCTION(BlueprintCallable, Category = "Shot Management")
+	static void AddProduction(FCinematicProduction Production);
+
+	/** Removes the Cinematic Production matching the input ProductionID from the Prodctuion Settings' list of productions */
+	UFUNCTION(BlueprintCallable, Category = "Shot Management")
+	static void DeleteProduction(FGuid ProductionID);
+
+	/** Renames the Cinematic Production matching the input ProductionID */
+	UFUNCTION(BlueprintCallable, Category = "Shot Management")
+	static void RenameProduction(FGuid ProductionID, FString NewName);
 };
