Shot Management: Adds an array of metadata descriptions to UCineAssemblySchema, which describe the Type, Key, and Default Value of a metadata field that you want to associate with assemblies made from that schema.
Add support to UCineAssembly to store its metadata in a json object, for simple queries and formatting.
Adds a detail customization for UCineAssembly to display additional property rows for each metadata field that it inherits from its base schema.
Adds a detail customization for FAssemblyMetadataDesc to display the correct widget type for the Default Value or each possible metadata type.
#rnx

diff --git a/Engine/Plugins/VirtualProduction/Shots/Source/ShotsEditor/Private/CineAssembly.cpp b/Engine/Plugins/VirtualProduction/Shots/Source/ShotsEditor/Private/CineAssembly.cpp
--- a/Engine/Plugins/VirtualProduction/Shots/Source/ShotsEditor/Private/CineAssembly.cpp
+++ b/Engine/Plugins/VirtualProduction/Shots/Source/ShotsEditor/Private/CineAssembly.cpp
@@ -1,12 +1,21 @@
 // Copyright Epic Games, Inc. All Rights Reserved.
 
 #include "CineAssembly.h"
+
+#include "Serialization/JsonReader.h"
+#include "Serialization/JsonSerializer.h"
+#include "Serialization/JsonWriter.h"
 #include "UObject/AssetRegistryTagsContext.h"
 
 #define LOCTEXT_NAMESPACE "CineAssembly"
 
 const FName UCineAssembly::AssetRegistryTag_AssemblyType = "AssemblyType";
 
+UCineAssembly::UCineAssembly()
+{
+	MetadataJsonObject = MakeShared<FJsonObject>();
+}
+
 const UCineAssemblySchema* UCineAssembly::GetSchema() const
 {
 	return BaseSchema.Get();
@@ -14,12 +23,39 @@ const UCineAssemblySchema* UCineAssembly::GetSchema() const
 
 void UCineAssembly::SetSchema(UCineAssemblySchema* InSchema)
 {
-	// If the assembly already has a base schema set, do not overwrite it
+	// Remove all metadata associated with the old schema before changing it
 	if (BaseSchema)
 	{
-		return;
+		for (const FAssemblyMetadataDesc& MetadataDesc : BaseSchema->AssemblyMetadata)
+		{
+			MetadataJsonObject->RemoveField(MetadataDesc.Key);
+		}
 	}
+
 	BaseSchema = InSchema;
+
+	// Add all metadata associated with the new schema (initialized to the default values for each field)
+	if (BaseSchema)
+	{
+		for (const FAssemblyMetadataDesc& MetadataDesc : BaseSchema->AssemblyMetadata)
+		{
+			if (MetadataDesc.DefaultValue.IsType<FString>())
+			{
+				const FString& DefaultValue = MetadataDesc.DefaultValue.Get<FString>();
+				MetadataJsonObject->SetStringField(MetadataDesc.Key, DefaultValue);
+			}
+			else if (MetadataDesc.DefaultValue.IsType<bool>())
+			{
+				const bool& DefaultValue = MetadataDesc.DefaultValue.Get<bool>();
+				MetadataJsonObject->SetBoolField(MetadataDesc.Key, DefaultValue);
+			}
+			else if (MetadataDesc.DefaultValue.IsType<double>())
+			{
+				const double& DefaultValue = MetadataDesc.DefaultValue.Get<double>();
+				MetadataJsonObject->SetNumberField(MetadataDesc.Key, DefaultValue);
+			}
+		}
+	}
 }
 
 void UCineAssembly::GetAssetRegistryTags(FAssetRegistryTagsContext Context) const
@@ -29,11 +65,13 @@ void UCineAssembly::GetAssetRegistryTags(FAssetRegistryTagsContext Context) cons
 	const FString AssemblyType = BaseSchema ? BaseSchema->SchemaName : TEXT("");
 	Context.AddTag(FAssetRegistryTag(AssetRegistryTag_AssemblyType, AssemblyType, FAssetRegistryTag::ETagType::TT_Alphabetical, FAssetRegistryTag::TD_None));
 
-	for (const TPair<FName, FString>& Pair : CustomMetadata)
+	// Add tags associated with the assembly metadata
+	for (const TPair<FString, TSharedPtr<FJsonValue>>& Pair : MetadataJsonObject->Values)
 	{
-		if (!Pair.Key.IsNone() && !Pair.Value.IsEmpty())
+		FString ValueString;
+		if (MetadataJsonObject->TryGetStringField(Pair.Key, ValueString))
 		{
-			Context.AddTag(FAssetRegistryTag(Pair.Key, Pair.Value, FAssetRegistryTag::ETagType::TT_Alphabetical, FAssetRegistryTag::TD_None));
+			Context.AddTag(FAssetRegistryTag(*Pair.Key, ValueString, FAssetRegistryTag::ETagType::TT_Alphabetical, FAssetRegistryTag::TD_None));
 		}
 	}
 }
@@ -48,4 +86,125 @@ void UCineAssembly::GetAssetRegistryTagMetadata(TMap<FName, FAssetRegistryTagMet
 	);
 }
 
+void UCineAssembly::PostEditChangeProperty(struct FPropertyChangedEvent& PropertyChangedEvent)
+{
+	Super::PostEditChangeProperty(PropertyChangedEvent);
+
+	const FName PropertyName = PropertyChangedEvent.GetPropertyName();
+
+	if (PropertyName == GET_MEMBER_NAME_CHECKED(UCineAssembly, InstanceMetadata))
+	{
+		UpdateInstanceMetadata();
+	}
+}
+
+void UCineAssembly::Serialize(FArchive& Ar)
+{
+	Super::Serialize(Ar);
+
+	FString JsonString;
+
+	if (Ar.IsSaving())
+ 	{
+ 		TSharedRef<TJsonWriter<>> JsonWriter = TJsonWriterFactory<>::Create(&JsonString);
+ 		FJsonSerializer::Serialize(MetadataJsonObject.ToSharedRef(), JsonWriter);
+ 	}
+
+	Ar << JsonString;
+
+	if (Ar.IsLoading())
+	{
+		TSharedRef<TJsonReader<>> JsonReader = TJsonReaderFactory<>::Create(JsonString);
+		FJsonSerializer::Deserialize(JsonReader, MetadataJsonObject);
+	}
+}
+
+FString UCineAssembly::GetFullMetadataString() const
+{
+	FString JsonString;
+
+	TSharedRef<TJsonWriter<>> JsonWriter = TJsonWriterFactory<>::Create(&JsonString);
+	FJsonSerializer::Serialize(MetadataJsonObject.ToSharedRef(), JsonWriter);
+
+	return JsonString;
+}
+
+void UCineAssembly::AddMetadataAsString(FString InKey, FString InValue)
+{
+	Modify();
+	MetadataJsonObject->SetStringField(InKey, InValue);
+}
+
+void UCineAssembly::AddMetadataAsBool(FString InKey, bool InValue)
+{
+	Modify();
+	MetadataJsonObject->SetBoolField(InKey, InValue);
+}
+
+void UCineAssembly::AddMetadataAsNumber(FString InKey, double InValue)
+{
+	Modify();
+	MetadataJsonObject->SetNumberField(InKey, InValue);
+}
+
+bool UCineAssembly::GetMetadataAsString(FString InKey, FString& OutValue) const
+{
+	if (!MetadataJsonObject->TryGetStringField(InKey, OutValue))
+	{
+		OutValue = TEXT("");
+		return false;
+	}
+	return true;
+}
+
+bool UCineAssembly::GetMetadataAsBool(FString InKey, bool& OutValue) const
+{
+	if (!MetadataJsonObject->TryGetBoolField(InKey, OutValue))
+	{
+		OutValue = false;
+		return false;
+	}
+	return true;
+}
+
+bool UCineAssembly::GetMetadataAsNumber(FString InKey, double& OutValue) const
+{
+	if (!MetadataJsonObject->TryGetNumberField(InKey, OutValue))
+	{
+		OutValue = 0;
+		return false;
+	}
+	return true;
+}
+
+void UCineAssembly::UpdateInstanceMetadata()
+{
+	// Copy our metadata key list so that we can remove keys as we encounter them in the map non-destructively
+	TArray<FName> InstanceMetadataKeysCopy = InstanceMetadataKeys;
+	for (const TPair<FName, FString>& Pair : InstanceMetadata)
+	{
+		if (!Pair.Key.IsNone())
+		{
+			if (InstanceMetadataKeys.Contains(Pair.Key))
+			{
+				// This is an existing metadata key that we are already tracking
+				InstanceMetadataKeysCopy.Remove(Pair.Key);
+			}
+			else
+			{
+				// This is a new metadata key that we were not previously tracking
+				InstanceMetadataKeys.Add(Pair.Key);
+			}
+
+			MetadataJsonObject->SetStringField(Pair.Key.ToString(), Pair.Value);
+		}
+	}
+
+	// If there are any keys remaining in our copy of the metadata key list, then those keys must have been removed from the instance metadata map
+	for (const FName& Key : InstanceMetadataKeysCopy)
+	{
+		MetadataJsonObject->RemoveField(Key.ToString());
+	}
+}
+
 #undef LOCTEXT_NAMESPACE
diff --git a/Engine/Plugins/VirtualProduction/Shots/Source/ShotsEditor/Private/CineAssembly.h b/Engine/Plugins/VirtualProduction/Shots/Source/ShotsEditor/Private/CineAssembly.h
--- a/Engine/Plugins/VirtualProduction/Shots/Source/ShotsEditor/Private/CineAssembly.h
+++ b/Engine/Plugins/VirtualProduction/Shots/Source/ShotsEditor/Private/CineAssembly.h
@@ -5,39 +5,76 @@
 #include "LevelSequence.h"
 
 #include "CineAssemblySchema.h"
+#include "Dom/JsonObject.h"
 
 #include "CineAssembly.generated.h"
 
 /** 
  * A cinematic building block that associates a level sequence with a level
  */
-UCLASS(HideCategories=Animation)
+UCLASS(BlueprintType, HideCategories=Animation)
 class UCineAssembly : public ULevelSequence
 {
 	GENERATED_BODY()
 
 public:
-	UCineAssembly() = default;
+	UCineAssembly();
+
+	friend class SCineAssemblyConfigWindow;
 
 	//~ Begin UObject Interface
 	virtual void GetAssetRegistryTags(FAssetRegistryTagsContext Context) const override;
 	virtual void GetAssetRegistryTagMetadata(TMap<FName, FAssetRegistryTagMetadata>& OutMetadata) const override;
+	virtual void PostEditChangeProperty(struct FPropertyChangedEvent& PropertyChangedEvent) override;
+	virtual void Serialize(FArchive& Ar) override;
 	//~ End UObject Interface
 
 	/** Get the base schema for this assembly */
 	const UCineAssemblySchema* GetSchema() const;
 
-	/** Assign the input schema to this assembly (if one has not already been assigned) */
+	/** Get all of the metadata for this assembly as a formatted JSON string */
+	UFUNCTION(BlueprintPure, Category = "Shot Management")
+	FString GetFullMetadataString() const;
+
+	/** Add a string as metadata to this assembly */
+	UFUNCTION(BlueprintCallable, Category = "Shot Management")
+	void AddMetadataAsString(FString InKey, FString InValue);
+
+	/** Add a boolean as metadata to this assembly */
+	UFUNCTION(BlueprintCallable, Category = "Shot Management")
+	void AddMetadataAsBool(FString InKey, bool InValue);
+
+	/** Add a double-precision floating point number as metadata to this assembly */
+	UFUNCTION(BlueprintCallable, Category = "Shot Management")
+	void AddMetadataAsNumber(FString InKey, double InValue);
+
+	/** Get the metadata value for the input key as a string (if it exists) */
+	UFUNCTION(BlueprintPure, Category = "Shot Management", meta = (ReturnDisplayName = "FoundKey"))
+	bool GetMetadataAsString(FString InKey, FString& OutValue) const;
+
+	/** Get the metadata value for the input key as a boolean (if it exists) */
+	UFUNCTION(BlueprintPure, Category = "Shot Management", meta = (ReturnDisplayName = "FoundKey"))
+	bool GetMetadataAsBool(FString InKey, bool& OutValue) const;
+
+	/** Get the metadata value for the input key as a double-precision floating-point number (if it exists) */
+	UFUNCTION(BlueprintPure, Category = "Shot Management", meta = (ReturnDisplayName = "FoundKey"))
+	bool GetMetadataAsNumber(FString InKey, double& OutValue) const;
+
+private:
+	/** Sets the base schema for this assembly and re-initializes the metadata inherited from the schema */
 	void SetSchema(UCineAssemblySchema* InSchema);
 
+	/** Update the underlying json object whenever keys/values in the instance metadata map are added/removed/modified */
+	void UpdateInstanceMetadata();
+
 public:
 	/** The level to open before opening this asset in Sequencer */
 	UPROPERTY(EditAnywhere, Category = "TargetLevel", meta = (AllowedClasses = "/Script/Engine.World"))
 	FSoftObjectPath Level;
 
 	/** User added metadata key/value pairs, which will be added as additional asset registry tags */
-	UPROPERTY(EditAnywhere, Category = "Metadata")
-	TMap<FName, FString> CustomMetadata;
+	UPROPERTY(EditAnywhere, Category = "Instance Metadata")
+	TMap<FName, FString> InstanceMetadata;
 
 	/** User-facing notes about this assembly asset */
 	UPROPERTY()
@@ -50,4 +87,10 @@ private:
 	/** The schema that was used as a base when creating this assembly (can be null if no schema was used) */
 	UPROPERTY()
 	TObjectPtr<UCineAssemblySchema> BaseSchema;
+
+	/** Copy of the keys present in the InstanceMetadata map, used to keep the json representation consistent with the map contents */
+	TArray<FName> InstanceMetadataKeys;
+
+	/** Json object responsible for storing the schema and instance metadata for this assembly */
+	TSharedPtr<FJsonObject> MetadataJsonObject;
 };
diff --git a/Engine/Plugins/VirtualProduction/Shots/Source/ShotsEditor/Private/CineAssemblyCustomization.cpp b/Engine/Plugins/VirtualProduction/Shots/Source/ShotsEditor/Private/CineAssemblyCustomization.cpp
new file mode 100644
--- /dev/null
+++ b/Engine/Plugins/VirtualProduction/Shots/Source/ShotsEditor/Private/CineAssemblyCustomization.cpp
@@ -0,0 +1,184 @@
+// Copyright Epic Games, Inc. All Rights Reserved.
+
+#include "CineAssemblyCustomization.h"
+
+#include "CineAssembly.h"
+#include "DetailLayoutBuilder.h"
+#include "DetailWidgetRow.h"
+#include "PropertyCustomizationHelpers.h"
+#include "Widgets/Input/SCheckBox.h"
+#include "Widgets/Input/SMultiLineEditableTextBox.h"
+#include "Widgets/Input/SNumericEntryBox.h"
+#include "Widgets/Layout/SBox.h"
+#include "Widgets/Text/STextBlock.h"
+
+#define LOCTEXT_NAMESPACE "CineAssemblyCustomization"
+
+TSharedRef<IDetailCustomization> FCineAssemblyCustomization::MakeInstance()
+{
+	return MakeShared<FCineAssemblyCustomization>();
+}
+
+void FCineAssemblyCustomization::CustomizeDetails(IDetailLayoutBuilder& DetailBuilder)
+{
+	TArray<TWeakObjectPtr<UObject>> CustomizedObjects;
+	DetailBuilder.GetObjectsBeingCustomized(CustomizedObjects);
+
+	// Ensure that we are only customizing one object
+	if (CustomizedObjects.Num() != 1)
+	{
+		return;
+	}
+
+	UCineAssembly* CustomizedCineAssembly = Cast<UCineAssembly>(CustomizedObjects[0]);
+
+	if (const UCineAssemblySchema* BaseSchema = CustomizedCineAssembly->GetSchema())
+	{
+		// 'Edit' the target level category to get it to appear first in the details view
+		DetailBuilder.EditCategory("TargetLevel", LOCTEXT("TargetLevelCategoryName", "Target Level"));
+
+		// Add a new category for Schema Metadata properties
+		const FText CategoryName = FText::Format(LOCTEXT("SchemaMetadataCategoryName", "{0} Metadata"), FText::FromString(BaseSchema->SchemaName));
+		IDetailCategoryBuilder& MetadataCategory = DetailBuilder.EditCategory("SchemaMetadata", CategoryName);
+
+		// Add a property row for each metadata struct in the customized Assembly's base schema
+		for (const FAssemblyMetadataDesc& MetadataDesc : BaseSchema->AssemblyMetadata)
+		{
+			if (MetadataDesc.Key.IsEmpty())
+			{
+				continue;
+			}
+
+			TSharedPtr<SWidget> ValueWidget;
+			if (MetadataDesc.Type == ECineAssemblyMetadataType::String)
+			{
+				ValueWidget = SNew(SBox)
+					.MaxDesiredHeight(120.0f)
+					[
+						SNew(SMultiLineEditableTextBox)
+							.Font(FAppStyle::GetFontStyle(TEXT("PropertyWindow.NormalFont")))
+							.AutoWrapText(true)
+							.Text_Lambda([CustomizedCineAssembly, MetadataDesc]()
+								{
+									FString Value;
+									if (!CustomizedCineAssembly->GetMetadataAsString(MetadataDesc.Key, Value))
+									{
+										Value = MetadataDesc.DefaultValue.Get<FString>();
+									}
+									return FText::FromString(Value);
+								})
+							.OnTextCommitted_Lambda([CustomizedCineAssembly, MetadataDesc](const FText& InText, ETextCommit::Type InCommitType)
+								{
+									CustomizedCineAssembly->AddMetadataAsString(MetadataDesc.Key, InText.ToString());
+								})
+					];
+			}
+			else if (MetadataDesc.Type == ECineAssemblyMetadataType::Bool)
+			{
+				ValueWidget = SNew(SCheckBox)
+					.IsChecked_Lambda([CustomizedCineAssembly, MetadataDesc]()
+						{
+							bool Value;
+							if (!CustomizedCineAssembly->GetMetadataAsBool(MetadataDesc.Key, Value))
+							{
+								Value = MetadataDesc.DefaultValue.Get<bool>();
+							}
+							return Value ? ECheckBoxState::Checked : ECheckBoxState::Unchecked;
+						})
+					.OnCheckStateChanged_Lambda([CustomizedCineAssembly, MetadataDesc](ECheckBoxState CheckBoxState)
+						{
+							const bool Value = (CheckBoxState == ECheckBoxState::Checked);
+							CustomizedCineAssembly->AddMetadataAsBool(MetadataDesc.Key, Value);
+						});
+			}
+			else if (MetadataDesc.Type == ECineAssemblyMetadataType::Number)
+			{
+				ValueWidget = SNew(SNumericEntryBox<double>)
+					.Font(FAppStyle::GetFontStyle(TEXT("PropertyWindow.NormalFont")))
+					.Value_Lambda([CustomizedCineAssembly, MetadataDesc]()
+						{
+							double Value;
+							if (!CustomizedCineAssembly->GetMetadataAsNumber(MetadataDesc.Key, Value))
+							{
+								Value = MetadataDesc.DefaultValue.Get<double>();
+							}
+							return Value;
+						})
+					.OnValueChanged_Lambda([CustomizedCineAssembly, MetadataDesc](double InValue)
+						{
+							CustomizedCineAssembly->AddMetadataAsNumber(MetadataDesc.Key, InValue);
+						});
+			}
+			else if (MetadataDesc.Type == ECineAssemblyMetadataType::AssetPath)
+			{
+				ValueWidget = SNew(SObjectPropertyEntryBox)
+					.AllowedClass(MetadataDesc.AssetClass.ResolveClass())
+					.ThumbnailPool(DetailBuilder.GetThumbnailPool())
+					.AllowCreate(true)
+					.ObjectPath_Lambda([CustomizedCineAssembly, MetadataDesc]()
+						{
+							FString Value;
+							if (!CustomizedCineAssembly->GetMetadataAsString(MetadataDesc.Key, Value))
+							{
+								Value = MetadataDesc.DefaultValue.Get<FString>();
+							}
+							return Value;
+						})
+					.OnObjectChanged_Lambda([CustomizedCineAssembly, MetadataDesc](const FAssetData& InAssetData)
+						{
+							// Store the path of the selected object as a string
+							CustomizedCineAssembly->AddMetadataAsString(MetadataDesc.Key, InAssetData.GetObjectPathString());
+						});
+			}
+			else if (MetadataDesc.Type == ECineAssemblyMetadataType::CineAssembly)
+			{
+				ValueWidget = SNew(SObjectPropertyEntryBox)
+					.AllowedClass(UCineAssembly::StaticClass())
+					.ThumbnailPool(DetailBuilder.GetThumbnailPool())
+					.AllowCreate(true)
+					.OnShouldFilterAsset_Lambda([MetadataDesc](const FAssetData& InAssetData) -> bool
+						{
+							const FAssetDataTagMapSharedView::FFindTagResult AssemblyType = InAssetData.TagsAndValues.FindTag(UCineAssembly::AssetRegistryTag_AssemblyType);
+							if (AssemblyType.IsSet())
+							{
+								return !AssemblyType.GetValue().Equals(MetadataDesc.SchemaType.GetAssetName());
+							}
+							return true;
+						})
+					.ObjectPath_Lambda([CustomizedCineAssembly, MetadataDesc]()
+						{
+							FString Value;
+							if (!CustomizedCineAssembly->GetMetadataAsString(MetadataDesc.Key, Value))
+							{
+								Value = MetadataDesc.DefaultValue.Get<FString>();
+							}
+							return Value;
+						})
+					.OnObjectChanged_Lambda([CustomizedCineAssembly, MetadataDesc](const FAssetData& InAssetData)
+						{
+							// Store the path of the selected object as a string
+							CustomizedCineAssembly->AddMetadataAsString(MetadataDesc.Key, InAssetData.GetObjectPathString());
+						});
+			}
+			else
+			{
+				checkNoEntry();
+			}
+
+			MetadataCategory.AddCustomRow(FText::FromString(MetadataDesc.Key))
+				.NameContent()
+				[
+					SNew(STextBlock)
+						.Text(FText::FromString(MetadataDesc.Key))
+						.Font(IDetailLayoutBuilder::GetDetailFont())
+				]
+				.ValueContent()
+				.HAlign(HAlign_Fill)
+				[
+					ValueWidget.ToSharedRef()
+				];
+		}
+	}
+}
+
+#undef LOCTEXT_NAMESPACE
diff --git a/Engine/Plugins/VirtualProduction/Shots/Source/ShotsEditor/Private/CineAssemblyCustomization.h b/Engine/Plugins/VirtualProduction/Shots/Source/ShotsEditor/Private/CineAssemblyCustomization.h
new file mode 100644
--- /dev/null
+++ b/Engine/Plugins/VirtualProduction/Shots/Source/ShotsEditor/Private/CineAssemblyCustomization.h
@@ -0,0 +1,19 @@
+// Copyright Epic Games, Inc. All Rights Reserved.
+
+#pragma once
+
+#include "IDetailCustomization.h"
+
+/**
+ * Detail customization for UCineAssembly
+ */
+class FCineAssemblyCustomization : public IDetailCustomization
+{
+public:
+	/** Makes a new instance of this detail layout class for a specific detail view requesting it */
+	static TSharedRef<IDetailCustomization> MakeInstance();
+
+	//~ Begin IDetailCustomization interface */
+	virtual void CustomizeDetails(IDetailLayoutBuilder& DetailBuilder) override;
+	//~ End IDetailCustomization interface */
+};
diff --git a/Engine/Plugins/VirtualProduction/Shots/Source/ShotsEditor/Private/CineAssemblyMetadataCustomization.cpp b/Engine/Plugins/VirtualProduction/Shots/Source/ShotsEditor/Private/CineAssemblyMetadataCustomization.cpp
new file mode 100644
--- /dev/null
+++ b/Engine/Plugins/VirtualProduction/Shots/Source/ShotsEditor/Private/CineAssemblyMetadataCustomization.cpp
@@ -0,0 +1,203 @@
+// Copyright Epic Games, Inc. All Rights Reserved.
+
+#include "CineAssemblyMetadataCustomization.h"
+
+#include "CineAssembly.h"
+#include "CineAssemblySchema.h"
+#include "IDetailChildrenBuilder.h"
+#include "IPropertyUtilities.h"
+#include "PropertyCustomizationHelpers.h"
+#include "Widgets/Input/SCheckBox.h"
+#include "Widgets/Input/SMultiLineEditableTextBox.h"
+#include "Widgets/Input/SNumericEntryBox.h"
+#include "Widgets/Layout/SBox.h"
+#include "Widgets/Layout/SWidgetSwitcher.h"
+#include "Widgets/Text/STextBlock.h"
+
+TSharedRef<IPropertyTypeCustomization> FCineAssemblyMetadataCustomization::MakeInstance()
+{
+	return MakeShareable(new FCineAssemblyMetadataCustomization);
+}
+
+void FCineAssemblyMetadataCustomization::CustomizeHeader(TSharedRef<IPropertyHandle> InPropertyHandle, FDetailWidgetRow& HeaderRow, IPropertyTypeCustomizationUtils& CustomizationUtils)
+{
+	HeaderRow.NameContent()
+		[
+			InPropertyHandle->CreatePropertyNameWidget()
+		]
+		.ValueContent()
+		[
+			InPropertyHandle->CreatePropertyValueWidget()
+		];
+}
+
+void FCineAssemblyMetadataCustomization::CustomizeChildren(TSharedRef<IPropertyHandle> PropertyHandle, IDetailChildrenBuilder& ChildBuilder, IPropertyTypeCustomizationUtils& CustomizationUtils)
+{
+	// Get the schema object that owns the metadata struct being customized
+	TArray<TWeakObjectPtr<UObject>> CustomizedObjects = CustomizationUtils.GetPropertyUtilities()->GetSelectedObjects();
+	if (CustomizedObjects.Num() != 1)
+	{
+		return;
+	}
+
+	UCineAssemblySchema* Schema = Cast<UCineAssemblySchema>(CustomizedObjects[0].Get());
+
+	// Add all of the existing reflected properties of the metadata struct
+	uint32 NumChildren;
+	PropertyHandle->GetNumChildren(NumChildren);
+
+	for (uint32 Index = 0; Index < NumChildren; ++Index)
+	{
+		TSharedRef<IPropertyHandle> ChildPropertyHandle = PropertyHandle->GetChildHandle(Index).ToSharedRef();
+		ChildBuilder.AddProperty(ChildPropertyHandle).ShowPropertyButtons(false);
+	}
+
+	// Add a "Default Value" property, based on the metadata type
+	int32 ArrayIndex = PropertyHandle->GetArrayIndex();
+	if (Schema->AssemblyMetadata.IsValidIndex(ArrayIndex))
+	{
+		FAssemblyMetadataDesc& MetadataDesc = Schema->AssemblyMetadata[ArrayIndex];
+
+		// Create a widget switcher that can display the appropriate widget based on the metadata type
+		TSharedPtr<IPropertyHandle> MetadataTypeHandle = PropertyHandle->GetChildHandle(GET_MEMBER_NAME_CHECKED(FAssemblyMetadataDesc, Type));
+
+		TSharedPtr<SWidgetSwitcher> DefaultValueWidget = SNew(SWidgetSwitcher)
+			.WidgetIndex_Lambda([MetadataTypeHandle]()
+				{
+					uint8 Value;
+					MetadataTypeHandle->GetValue(Value);
+					return Value;
+				});
+
+		DefaultValueWidget->AddSlot((uint8)ECineAssemblyMetadataType::String)
+			[
+				SNew(SBox)
+					.MaxDesiredHeight(120.0f)
+					[
+						SNew(SMultiLineEditableTextBox)
+							.Font(FAppStyle::GetFontStyle(TEXT("PropertyWindow.NormalFont")))
+							.AutoWrapText(true)
+							.Text_Lambda([&MetadataDesc]() -> FText
+								{
+									if (const FString* Value = MetadataDesc.DefaultValue.TryGet<FString>())
+									{
+										return FText::FromString(*Value);
+									}
+									return FText::GetEmpty();
+								})
+							.OnTextCommitted_Lambda([&MetadataDesc](const FText& InText, ETextCommit::Type InCommitType)
+								{
+									MetadataDesc.DefaultValue.Set<FString>(InText.ToString());
+								})
+					]
+			];
+
+		DefaultValueWidget->AddSlot((uint8)ECineAssemblyMetadataType::Bool)
+			[
+				SNew(SCheckBox)
+					.IsChecked_Lambda([&MetadataDesc]() -> ECheckBoxState
+						{
+							if (const bool* Value = MetadataDesc.DefaultValue.TryGet<bool>())
+							{
+								return *Value ? ECheckBoxState::Checked : ECheckBoxState::Unchecked;
+							}
+							return ECheckBoxState::Unchecked;
+						})
+					.OnCheckStateChanged_Lambda([&MetadataDesc](ECheckBoxState CheckBoxState)
+						{
+							const bool Value = (CheckBoxState == ECheckBoxState::Checked);
+							MetadataDesc.DefaultValue.Set<bool>(Value);
+						})
+			];
+
+		DefaultValueWidget->AddSlot((uint8)ECineAssemblyMetadataType::Number)
+			[
+				SNew(SNumericEntryBox<double>)
+					.Font(FAppStyle::GetFontStyle(TEXT("PropertyWindow.NormalFont")))
+					.Value_Lambda([&MetadataDesc]() -> double
+						{
+							if (const double* Value = MetadataDesc.DefaultValue.TryGet<double>())
+							{
+								return *Value;
+							}
+							return 0;
+						})
+					.OnValueChanged_Lambda([&MetadataDesc](double InValue)
+						{
+							MetadataDesc.DefaultValue.Set<double>(InValue);
+						})
+			];
+
+		DefaultValueWidget->AddSlot((uint8)ECineAssemblyMetadataType::AssetPath)
+			[
+				SNew(SObjectPropertyEntryBox)
+					.ThumbnailPool(CustomizationUtils.GetThumbnailPool())
+					.AllowCreate(true)
+					.OnShouldFilterAsset_Lambda([&MetadataDesc](const FAssetData& InAssetData) -> bool
+						{
+							// Filter out all assets that do not match the selected Asset Class for this metadata struct
+							if (MetadataDesc.AssetClass.IsNull())
+							{
+								return true;
+							}
+							return (InAssetData.AssetClassPath != MetadataDesc.AssetClass.GetAssetPath());
+						})
+					.ObjectPath_Lambda([&MetadataDesc]() -> FString
+						{
+							if (const FString* Value = MetadataDesc.DefaultValue.TryGet<FString>())
+							{
+								return *Value;
+							}
+							return FString();
+						})
+					.OnObjectChanged_Lambda([&MetadataDesc](const FAssetData& InAssetData)
+						{
+							MetadataDesc.DefaultValue.Set<FString>(InAssetData.GetObjectPathString());
+						})
+			];
+
+		DefaultValueWidget->AddSlot((uint8)ECineAssemblyMetadataType::CineAssembly)
+			[
+				SNew(SObjectPropertyEntryBox)
+					.AllowedClass(UCineAssembly::StaticClass())
+					.ThumbnailPool(CustomizationUtils.GetThumbnailPool())
+					.AllowCreate(true)
+					.OnShouldFilterAsset_Lambda([&MetadataDesc](const FAssetData& InAssetData) -> bool
+						{
+							// Filter out all Cine Assembly assets that do not match the selected Schema Type for this metadata struct
+							const FAssetDataTagMapSharedView::FFindTagResult AssemblyType = InAssetData.TagsAndValues.FindTag(UCineAssembly::AssetRegistryTag_AssemblyType);
+							if (AssemblyType.IsSet())
+							{
+								return !AssemblyType.GetValue().Equals(MetadataDesc.SchemaType.GetAssetName());
+							}
+							return true;
+						})
+					.ObjectPath_Lambda([&MetadataDesc]() -> FString
+						{
+							if (const FString* Value = MetadataDesc.DefaultValue.TryGet<FString>())
+							{
+								return *Value;
+							}
+							return FString();
+						})
+					.OnObjectChanged_Lambda([&MetadataDesc](const FAssetData& InAssetData)
+						{
+							MetadataDesc.DefaultValue.Set<FString>(InAssetData.GetObjectPathString());
+						})
+			];
+
+		ChildBuilder.AddCustomRow(FText::GetEmpty())
+			.NameContent()
+			[
+				SNew(STextBlock)
+					.Text(NSLOCTEXT("CineAssemblyMetadataCustomization", "DefaultValueText", "Default Value"))
+					.Font(FAppStyle::GetFontStyle(TEXT("PropertyWindow.NormalFont")))
+			]
+			.ValueContent()
+			[
+				DefaultValueWidget.ToSharedRef()
+			];
+	}
+}
+
+#undef LOCTEXT_NAMESPACE
diff --git a/Engine/Plugins/VirtualProduction/Shots/Source/ShotsEditor/Private/CineAssemblyMetadataCustomization.h b/Engine/Plugins/VirtualProduction/Shots/Source/ShotsEditor/Private/CineAssemblyMetadataCustomization.h
new file mode 100644
--- /dev/null
+++ b/Engine/Plugins/VirtualProduction/Shots/Source/ShotsEditor/Private/CineAssemblyMetadataCustomization.h
@@ -0,0 +1,16 @@
+// Copyright Epic Games, Inc. All Rights Reserved.
+
+#pragma once
+
+#include "IPropertyTypeCustomization.h"
+
+class FCineAssemblyMetadataCustomization : public IPropertyTypeCustomization
+{
+public:
+	static TSharedRef<IPropertyTypeCustomization> MakeInstance();
+
+	// Begin IDetailCustomization interface
+	virtual void CustomizeHeader(TSharedRef<IPropertyHandle> PropertyHandle, FDetailWidgetRow& HeaderRow, IPropertyTypeCustomizationUtils& CustomizationUtils) override;
+	virtual void CustomizeChildren(TSharedRef<IPropertyHandle> PropertyHandle, IDetailChildrenBuilder& ChildBuilder, IPropertyTypeCustomizationUtils& CustomizationUtils) override;
+	//~ End IDetailCustomization interface
+};
diff --git a/Engine/Plugins/VirtualProduction/Shots/Source/ShotsEditor/Private/CineAssemblySchema.cpp b/Engine/Plugins/VirtualProduction/Shots/Source/ShotsEditor/Private/CineAssemblySchema.cpp
--- a/Engine/Plugins/VirtualProduction/Shots/Source/ShotsEditor/Private/CineAssemblySchema.cpp
+++ b/Engine/Plugins/VirtualProduction/Shots/Source/ShotsEditor/Private/CineAssemblySchema.cpp
@@ -14,6 +14,15 @@ UCineAssemblySchema::UCineAssemblySchema()
 	ThumbnailBrush->ImageSize = FVector2D(64.0f, 64.0f);
 }
 
+void UCineAssemblySchema::Serialize(FArchive& Ar)
+{
+	Super::Serialize(Ar);
+	for (FAssemblyMetadataDesc& Desc : AssemblyMetadata)
+	{
+		Ar << Desc.DefaultValue;
+	}
+}
+
 void UCineAssemblySchema::PostLoad()
 {
 	Super::PostLoad();
diff --git a/Engine/Plugins/VirtualProduction/Shots/Source/ShotsEditor/Private/CineAssemblySchema.h b/Engine/Plugins/VirtualProduction/Shots/Source/ShotsEditor/Private/CineAssemblySchema.h
--- a/Engine/Plugins/VirtualProduction/Shots/Source/ShotsEditor/Private/CineAssemblySchema.h
+++ b/Engine/Plugins/VirtualProduction/Shots/Source/ShotsEditor/Private/CineAssemblySchema.h
@@ -7,6 +7,43 @@
 
 #include "CineAssemblySchema.generated.h"
 
+/** The types of assembly metadata supported by Cine Assembly Schemas */
+UENUM()
+enum class ECineAssemblyMetadataType : uint8
+{
+	String = 0,
+	Bool,
+	Number,
+	AssetPath,
+	CineAssembly
+};
+
+/** Structure defining a single metadata field that can be associated with an assembly built from this schema, including its type, key, and default value */
+USTRUCT()
+struct FAssemblyMetadataDesc
+{
+	GENERATED_BODY()
+
+	/** Metadata type */
+	UPROPERTY(EditAnywhere, Category = "Metadata")
+	ECineAssemblyMetadataType Type = ECineAssemblyMetadataType::String;
+
+	/** The key associated with this field */
+	UPROPERTY(EditAnywhere, Category = "Metadata")
+	FString Key;
+
+	/** For AssetPath types, the class to restrict the value of this metadata field to */
+	UPROPERTY(EditAnywhere, Category = "Metadata", meta = (AllowAbstract = false, HideViewOptions, EditCondition = "Type == ECineAssemblyMetadataType::AssetPath", EditConditionHides))
+	FSoftClassPath AssetClass;
+
+	/** For CineAssembly types, the schema type to restrict the value of this metadata field to */
+	UPROPERTY(EditAnywhere, Category = "Metadata", meta = (AllowedClasses="/Script/ShotsEditor.CineAssemblySchema", EditCondition = "Type == ECineAssemblyMetadataType::CineAssembly", EditConditionHides))
+	FSoftObjectPath SchemaType;
+
+	/** The default value for this metadata field */
+	TVariant<FString, bool, double> DefaultValue;
+};
+
 /**
  * A template object for building different Cine Assembly types
  */
@@ -19,6 +56,7 @@ public:
 	UCineAssemblySchema();
 
 	//~ Begin UObject overrides
+	virtual void Serialize(FArchive& Ar) override;
 	virtual void PostLoad() override;
 	virtual void PostEditChangeProperty(struct FPropertyChangedEvent& PropertyChangedEvent) override;
 	//~ End UObject overrides
@@ -49,6 +87,10 @@ public:
 	UPROPERTY(EditAnywhere, Category = "Default")
 	TObjectPtr<UTexture2D> ThumbnailImage;
 
+	/** List of metadata fields that should be automatically added to assemblies made from this schema */
+	UPROPERTY(EditAnywhere, Category = "Metadata")
+	TArray<FAssemblyMetadataDesc> AssemblyMetadata;
+
 	/** Paths of subsequence assets that should be created for assemblies that use this schema, relative to the path of the top-level assembly */
 	UPROPERTY()
 	TArray<FString> SubsequencesToCreate;
diff --git a/Engine/Plugins/VirtualProduction/Shots/Source/ShotsEditor/Private/ShotsEditorModule.cpp b/Engine/Plugins/VirtualProduction/Shots/Source/ShotsEditor/Private/ShotsEditorModule.cpp
--- a/Engine/Plugins/VirtualProduction/Shots/Source/ShotsEditor/Private/ShotsEditorModule.cpp
+++ b/Engine/Plugins/VirtualProduction/Shots/Source/ShotsEditor/Private/ShotsEditorModule.cpp
@@ -3,6 +3,9 @@
 #include "Modules/ModuleManager.h"
 
 #include "CineAssembly.h"
+#include "CineAssemblySchema.h"
+#include "CineAssemblyCustomization.h"
+#include "CineAssemblyMetadataCustomization.h"
 #include "ContentBrowserMenuContexts.h"
 #include "Editor.h"
 #include "Framework/Application/SlateApplication.h"
@@ -49,6 +52,16 @@ void FShotsEditorModule::StartupModule()
 		FOnGetDetailCustomizationInstance::CreateStatic(&FProductionSettingsCustomization::MakeInstance)
 	);
 
+	PropertyEditorModule.RegisterCustomClassLayout(
+		UCineAssembly::StaticClass()->GetFName(),
+		FOnGetDetailCustomizationInstance::CreateStatic(&FCineAssemblyCustomization::MakeInstance)
+	);
+
+	PropertyEditorModule.RegisterCustomPropertyTypeLayout(
+		FAssemblyMetadataDesc::StaticStruct()->GetFName(),
+		FOnGetPropertyTypeCustomizationInstance::CreateStatic(&FCineAssemblyMetadataCustomization::MakeInstance)
+	);
+
 	// Add a new entry to the Tools->Cinematics menu to spawn the Production Wizard
 	FGlobalTabmanager::Get()->RegisterNomadTabSpawner(ProductionWizardTabName, FOnSpawnTab::CreateRaw(this, &FShotsEditorModule::SpawnProductionWizard))
 		.SetGroup(WorkspaceMenu::GetMenuStructure().GetLevelEditorCinematicsCategory())
diff --git a/Engine/Plugins/VirtualProduction/Shots/Source/ShotsEditor/Private/UI/CineAssembly/SCineAssemblyConfigWindow.cpp b/Engine/Plugins/VirtualProduction/Shots/Source/ShotsEditor/Private/UI/CineAssembly/SCineAssemblyConfigWindow.cpp
--- a/Engine/Plugins/VirtualProduction/Shots/Source/ShotsEditor/Private/UI/CineAssembly/SCineAssemblyConfigWindow.cpp
+++ b/Engine/Plugins/VirtualProduction/Shots/Source/ShotsEditor/Private/UI/CineAssembly/SCineAssemblyConfigWindow.cpp
@@ -155,6 +155,9 @@ void SCineAssemblyConfigWindow::OnSchemaSelected(const FAssetData& AssetData)
 	}
 
 	CineAssemblyToConfigure->SetSchema(SelectedSchema);
+
+	// The details view needs to be redrawn to show the new metadata fields from the selected schema
+	DetailsView->ForceRefresh();
 }
 
 TSharedRef<SWidget> SCineAssemblyConfigWindow::MakeInfoPanel()
@@ -300,10 +303,10 @@ TSharedRef<SWidget> SCineAssemblyConfigWindow::MakeDetailsWidget()
 	DetailsViewArgs.bAllowSearch = false;
 	DetailsViewArgs.bHideSelectionTip = true;
 
-	TSharedRef<IDetailsView> DetailsView = PropertyEditor.CreateDetailView(DetailsViewArgs);
+	DetailsView = PropertyEditor.CreateDetailView(DetailsViewArgs);
 	DetailsView->SetObject(CineAssemblyToConfigure, true);
 
-	return DetailsView;
+	return DetailsView.ToSharedRef();
 }
 
 TSharedRef<SWidget> SCineAssemblyConfigWindow::MakeNotesWidget()
diff --git a/Engine/Plugins/VirtualProduction/Shots/Source/ShotsEditor/Private/UI/CineAssembly/SCineAssemblyConfigWindow.h b/Engine/Plugins/VirtualProduction/Shots/Source/ShotsEditor/Private/UI/CineAssembly/SCineAssemblyConfigWindow.h
--- a/Engine/Plugins/VirtualProduction/Shots/Source/ShotsEditor/Private/UI/CineAssembly/SCineAssemblyConfigWindow.h
+++ b/Engine/Plugins/VirtualProduction/Shots/Source/ShotsEditor/Private/UI/CineAssembly/SCineAssemblyConfigWindow.h
@@ -6,6 +6,7 @@
 
 #include "CineAssembly.h"
 #include "CineAssemblySchema.h"
+#include "IDetailsView.h"
 #include "Widgets/Layout/SWidgetSwitcher.h"
 #include "Widgets/SWindow.h"
 
@@ -64,6 +65,9 @@ private:
 	/** Switcher that controls which tab widget is currently visible */
 	TSharedPtr<SWidgetSwitcher> TabSwitcher;
 
+	/** Details View displaying the reflected properties of the Cine Assembly being configured */ 
+	TSharedPtr<IDetailsView> DetailsView;
+
 	/** Whether a new asset should be created after closing this window */
 	bool bShouldCreateNewAsset = false;
 
