ShotManagement: Add NamingTokens class for Cinematic Assemblies.
Add Default Assembly Name to the schema definition, which accepts tokens.
Add an Assembly Name field to the assembly config window, which will display the evaluated token text entereted into it.
Change the Cine Assembly Factory to use the configured Assembly Name when creating new assets.
#rnx

diff --git a/Engine/Plugins/VirtualProduction/Shots/Source/ShotsEditor/Private/CineAssembly.cpp b/Engine/Plugins/VirtualProduction/Shots/Source/ShotsEditor/Private/CineAssembly.cpp
--- a/Engine/Plugins/VirtualProduction/Shots/Source/ShotsEditor/Private/CineAssembly.cpp
+++ b/Engine/Plugins/VirtualProduction/Shots/Source/ShotsEditor/Private/CineAssembly.cpp
@@ -2,6 +2,8 @@
 
 #include "CineAssembly.h"
 
+#include "CineAssemblyNamingTokens.h"
+#include "NamingTokensEngineSubsystem.h"
 #include "Serialization/JsonReader.h"
 #include "Serialization/JsonSerializer.h"
 #include "Serialization/JsonWriter.h"
@@ -116,6 +118,12 @@ void UCineAssembly::Serialize(FArchive& Ar)
 	{
 		TSharedRef<TJsonReader<>> JsonReader = TJsonReaderFactory<>::Create(JsonString);
 		FJsonSerializer::Deserialize(JsonReader, MetadataJsonObject);
+
+		// After the JsonObject has been loaded, add a naming token for each of its keys
+		for (const TPair<FString, TSharedPtr<FJsonValue>>& Pair : MetadataJsonObject->Values)
+		{
+			AddMetadataNamingToken(Pair.Key);
+		}
 	}
 }
 
@@ -133,18 +141,21 @@ void UCineAssembly::AddMetadataAsString(FString InKey, FString InValue)
 {
 	Modify();
 	MetadataJsonObject->SetStringField(InKey, InValue);
+	AddMetadataNamingToken(InKey);
 }
 
 void UCineAssembly::AddMetadataAsBool(FString InKey, bool InValue)
 {
 	Modify();
 	MetadataJsonObject->SetBoolField(InKey, InValue);
+	AddMetadataNamingToken(InKey);
 }
 
 void UCineAssembly::AddMetadataAsNumber(FString InKey, double InValue)
 {
 	Modify();
 	MetadataJsonObject->SetNumberField(InKey, InValue);
+	AddMetadataNamingToken(InKey);
 }
 
 bool UCineAssembly::GetMetadataAsString(FString InKey, FString& OutValue) const
@@ -196,7 +207,7 @@ void UCineAssembly::UpdateInstanceMetadata()
 				InstanceMetadataKeys.Add(Pair.Key);
 			}
 
-			MetadataJsonObject->SetStringField(Pair.Key.ToString(), Pair.Value);
+			AddMetadataAsString(Pair.Key.ToString(), Pair.Value);
 		}
 	}
 
@@ -207,4 +218,12 @@ void UCineAssembly::UpdateInstanceMetadata()
 	}
 }
 
+void UCineAssembly::AddMetadataNamingToken(const FString& InKey)
+{
+	UNamingTokensEngineSubsystem* NamingTokensSubsystem = GEngine->GetEngineSubsystem<UNamingTokensEngineSubsystem>();
+	UCineAssemblyNamingTokens* CineAssemblyNamingTokens = Cast<UCineAssemblyNamingTokens>(NamingTokensSubsystem->GetNamingTokens(UCineAssemblyNamingTokens::TokenNamespace));
+
+	CineAssemblyNamingTokens->AddMetadataToken(InKey);
+}
+
 #undef LOCTEXT_NAMESPACE
diff --git a/Engine/Plugins/VirtualProduction/Shots/Source/ShotsEditor/Private/CineAssembly.h b/Engine/Plugins/VirtualProduction/Shots/Source/ShotsEditor/Private/CineAssembly.h
--- a/Engine/Plugins/VirtualProduction/Shots/Source/ShotsEditor/Private/CineAssembly.h
+++ b/Engine/Plugins/VirtualProduction/Shots/Source/ShotsEditor/Private/CineAssembly.h
@@ -6,6 +6,7 @@
 
 #include "CineAssemblySchema.h"
 #include "Dom/JsonObject.h"
+#include "UObject/TemplateString.h"
 
 #include "CineAssembly.generated.h"
 
@@ -67,7 +68,14 @@ private:
 	/** Update the underlying json object whenever keys/values in the instance metadata map are added/removed/modified */
 	void UpdateInstanceMetadata();
 
+	/** Adds a new metadata key to the list of supported naming tokens for assemblies */
+	void AddMetadataNamingToken(const FString& InKey);
+
 public:
+	/** The assembly name, which supports tokens */
+	UPROPERTY()
+	FTemplateString AssemblyName;
+
 	/** The level to open before opening this asset in Sequencer */
 	UPROPERTY(EditAnywhere, Category = "TargetLevel", meta = (AllowedClasses = "/Script/Engine.World"))
 	FSoftObjectPath Level;
diff --git a/Engine/Plugins/VirtualProduction/Shots/Source/ShotsEditor/Private/CineAssemblyFactory.cpp b/Engine/Plugins/VirtualProduction/Shots/Source/ShotsEditor/Private/CineAssemblyFactory.cpp
--- a/Engine/Plugins/VirtualProduction/Shots/Source/ShotsEditor/Private/CineAssemblyFactory.cpp
+++ b/Engine/Plugins/VirtualProduction/Shots/Source/ShotsEditor/Private/CineAssemblyFactory.cpp
@@ -4,12 +4,14 @@
 
 #include "AssetRegistry/AssetRegistryModule.h"
 #include "AssetToolsModule.h"
+#include "CineAssemblyNamingTokens.h"
 #include "Framework/Application/SlateApplication.h"
 #include "HAL/FileManager.h"
 #include "Interfaces/IMainFrameModule.h"
 #include "Modules/ModuleManager.h"
 #include "MovieScene.h"
 #include "MovieSceneToolsProjectSettings.h"
+#include "NamingTokensEngineSubsystem.h"
 #include "Tracks/MovieSceneSubTrack.h"
 #include "UI/CineAssembly/SCineAssemblyConfigWindow.h"
 #include "UObject/Package.h"
@@ -38,6 +40,10 @@ UObject* UCineAssemblyFactory::FactoryCreateNew(UClass* Class, UObject* InParent
 		CineAssemblyToConfigure->Rename(*Name.ToString(), InParent);
 		CineAssemblyToConfigure->SetFlags(Flags);
 		CineAssemblyToConfigure->ClearFlags(RF_Transient);
+
+		// Unregister the temporary default name that we registered earlier for this schema instance
+		IAssetTools& AssetTools = FModuleManager::GetModuleChecked<FAssetToolsModule>("AssetTools").Get();
+		AssetTools.UnregisterDefaultAssetNameForClass(UCineAssembly::StaticClass());
 	}
 	else
 	{
@@ -130,6 +136,28 @@ bool UCineAssemblyFactory::ConfigureProperties()
 		return false;
 	}
 
+	// One of the properties configured by the user in the config window is the assembly name.
+	// We want the actual asset name to match the value of this property, and we don't want the user to have to retype it.
+	// This workaround allows us to temporarily set a default name for the assembly asset. 
+	// After FactoryCreateNew() is finished, we will unregister this default name override. 
+	if (!CineAssemblyToConfigure->AssemblyName.Template.IsEmpty())
+	{
+		// Add the cine assembly being created to the naming token context
+		UCineAssemblyNamingTokensContext* NamingTokenContext = NewObject<UCineAssemblyNamingTokensContext>();
+		NamingTokenContext->Assembly = CineAssemblyToConfigure;
+
+		FNamingTokenFilterArgs FilterArgs;
+		FilterArgs.AdditionalNamespacesToInclude.Add(UCineAssemblyNamingTokens::TokenNamespace);
+
+		UNamingTokensEngineSubsystem* NamingTokensSubsystem = GEngine->GetEngineSubsystem<UNamingTokensEngineSubsystem>();
+		FNamingTokenResultData Result = NamingTokensSubsystem->EvaluateTokenString(CineAssemblyToConfigure->AssemblyName.Template, FilterArgs, { NamingTokenContext });
+
+		CineAssemblyToConfigure->AssemblyName.Resolved = Result.EvaluatedText;
+
+		IAssetTools& AssetTools = FModuleManager::GetModuleChecked<FAssetToolsModule>("AssetTools").Get();
+		AssetTools.RegisterDefaultAssetNameForClass(UCineAssembly::StaticClass(), CineAssemblyToConfigure->AssemblyName.Resolved.ToString());
+	}
+
 	return true;
 }
 
diff --git a/Engine/Plugins/VirtualProduction/Shots/Source/ShotsEditor/Private/CineAssemblyMetadataCustomization.cpp b/Engine/Plugins/VirtualProduction/Shots/Source/ShotsEditor/Private/CineAssemblyMetadataCustomization.cpp
--- a/Engine/Plugins/VirtualProduction/Shots/Source/ShotsEditor/Private/CineAssemblyMetadataCustomization.cpp
+++ b/Engine/Plugins/VirtualProduction/Shots/Source/ShotsEditor/Private/CineAssemblyMetadataCustomization.cpp
@@ -61,6 +61,28 @@ void FCineAssemblyMetadataCustomization::CustomizeChildren(TSharedRef<IPropertyH
 		// Create a widget switcher that can display the appropriate widget based on the metadata type
 		TSharedPtr<IPropertyHandle> MetadataTypeHandle = PropertyHandle->GetChildHandle(GET_MEMBER_NAME_CHECKED(FAssemblyMetadataDesc, Type));
 
+		MetadataTypeHandle->SetOnPropertyValueChanged(FSimpleDelegate::CreateLambda([&MetadataDesc]()
+			{
+				if ((MetadataDesc.Type == ECineAssemblyMetadataType::String) ||
+					(MetadataDesc.Type == ECineAssemblyMetadataType::AssetPath) ||
+					(MetadataDesc.Type == ECineAssemblyMetadataType::CineAssembly))
+				{
+					MetadataDesc.DefaultValue.Set<FString>(TEXT(""));
+				}
+				else if (MetadataDesc.Type == ECineAssemblyMetadataType::Bool)
+				{
+					MetadataDesc.DefaultValue.Set<bool>(false);
+				}
+				else if (MetadataDesc.Type == ECineAssemblyMetadataType::Number)
+				{
+					MetadataDesc.DefaultValue.Set<double>(0);
+				}
+				else
+				{
+					checkNoEntry();
+				}
+			}));
+
 		TSharedPtr<SWidgetSwitcher> DefaultValueWidget = SNew(SWidgetSwitcher)
 			.WidgetIndex_Lambda([MetadataTypeHandle]()
 				{
diff --git a/Engine/Plugins/VirtualProduction/Shots/Source/ShotsEditor/Private/CineAssemblyNamingTokens.cpp b/Engine/Plugins/VirtualProduction/Shots/Source/ShotsEditor/Private/CineAssemblyNamingTokens.cpp
new file mode 100644
--- /dev/null
+++ b/Engine/Plugins/VirtualProduction/Shots/Source/ShotsEditor/Private/CineAssemblyNamingTokens.cpp
@@ -0,0 +1,137 @@
+// Copyright Epic Games, Inc. All Rights Reserved.
+
+#include "CineAssemblyNamingTokens.h"
+
+#include "Algo/IndexOf.h"
+
+#define LOCTEXT_NAMESPACE "CineAssemblyNamingTokens"
+
+FString UCineAssemblyNamingTokens::TokenNamespace = TEXT("cat");
+
+UCineAssemblyNamingTokens::UCineAssemblyNamingTokens()
+{
+	Namespace = TokenNamespace;
+}
+
+void UCineAssemblyNamingTokens::OnCreateDefaultTokens(TArray<FNamingTokenData>& Tokens)
+{
+	Super::OnCreateDefaultTokens(Tokens);
+
+	auto AssemblyNameTokenFunc = [](TWeakObjectPtr<UCineAssembly> InAssembly) -> FText
+		{
+			if (InAssembly.IsValid())
+			{
+				return InAssembly->AssemblyName.Resolved;
+			}
+			return FText::GetEmpty();
+		};
+
+	auto SchemaTokenFunc = [](TWeakObjectPtr<UCineAssembly> InAssembly) -> FText
+		{
+			if (InAssembly.IsValid())
+			{
+				if (const UCineAssemblySchema* Schema = InAssembly->GetSchema())
+				{
+					return FText::FromString(Schema->SchemaName);
+				}
+			}
+			return FText::GetEmpty();
+		};
+
+	auto LevelTokenFunc = [](TWeakObjectPtr<UCineAssembly> InAssembly) -> FText
+		{
+			if (InAssembly.IsValid() && InAssembly->Level.IsValid())
+			{
+				return FText::FromString(InAssembly->Level.GetAssetName());
+			}
+			return FText::GetEmpty();
+		};
+
+	FNamingTokenData AssemblyNameToken;
+	AssemblyNameToken.TokenKey = TEXT("assembly_name");
+	AssemblyNameToken.DisplayName = LOCTEXT("AssemblyNameToken", "Assembly Name");
+	AssemblyNameToken.TokenProcessorNative.BindLambda([this, AssemblyNameTokenFunc]() { return ExecuteTokenFunc(AssemblyNameTokenFunc); });
+
+	FNamingTokenData SchemaToken;
+	SchemaToken.TokenKey = TEXT("schema");
+	SchemaToken.DisplayName = LOCTEXT("SchemaToken", "Base Schema");
+	SchemaToken.TokenProcessorNative.BindLambda([this, SchemaTokenFunc]() { return ExecuteTokenFunc(SchemaTokenFunc); });
+
+	FNamingTokenData LevelToken;
+	LevelToken.TokenKey = TEXT("level");
+	LevelToken.DisplayName = LOCTEXT("TargetLevelTokenName", "Target Level");
+	LevelToken.TokenProcessorNative.BindLambda([this, LevelTokenFunc]() { return ExecuteTokenFunc(LevelTokenFunc); });
+
+	Tokens.Add(AssemblyNameToken);
+	Tokens.Add(SchemaToken);
+	Tokens.Add(LevelToken);
+}
+
+void UCineAssemblyNamingTokens::OnPreEvaluate_Implementation(const FNamingTokensEvaluationData& InEvaluationData)
+{
+	Super::OnPreEvaluate_Implementation(InEvaluationData);
+
+	UCineAssemblyNamingTokensContext* MatchingContext = nullptr;
+	InEvaluationData.Contexts.FindItemByClass<UCineAssemblyNamingTokensContext>(&MatchingContext);
+	Context = MatchingContext;
+}
+
+void UCineAssemblyNamingTokens::OnPostEvaluate_Implementation()
+{
+	Super::OnPostEvaluate_Implementation();
+	Context = nullptr;
+}
+
+void UCineAssemblyNamingTokens::SetDefaultCineAssembly(UCineAssembly* InAssembly)
+{
+	DefaultAssembly = InAssembly;
+}
+
+void UCineAssemblyNamingTokens::AddMetadataToken(const FString& InTokenKey)
+{
+	// If the token already exists, remove it, and replace it with the new one
+	const int32 TokenIndex = Algo::IndexOfBy(CustomTokens, InTokenKey, &FNamingTokenData::TokenKey);
+	if (CustomTokens.IsValidIndex(TokenIndex))
+	{
+		CustomTokens.RemoveAt(TokenIndex);
+	}
+
+	auto MetadataTokenFunc = [InTokenKey](TWeakObjectPtr<UCineAssembly> InAssembly) -> FText
+		{
+			if (InAssembly.IsValid())
+			{
+				FString ValueString;
+				if (InAssembly->GetMetadataAsString(InTokenKey, ValueString))
+				{
+					return FText::FromString(ValueString);
+				}
+			}
+
+			return FText::GetEmpty();
+		};
+
+	FNamingTokenData NewToken;
+	NewToken.TokenKey = InTokenKey;
+	NewToken.DisplayName = FText::Format(LOCTEXT("MetadataTokenDisplayName", "{0} Metadata"), FText::FromString(InTokenKey));
+	NewToken.TokenProcessorNative.BindLambda([this, MetadataTokenFunc]() { return ExecuteTokenFunc(MetadataTokenFunc); });
+
+	CustomTokens.Add(NewToken);
+}
+
+FText UCineAssemblyNamingTokens::ExecuteTokenFunc(TFunction<FText(TWeakObjectPtr<UCineAssembly>)> TokenFunc)
+{
+	// Try evaluating the token function using the assembly from the context. If there was no context, try using the default assembly. 
+	if (Context)
+	{
+		return TokenFunc(Context->Assembly);
+	}
+	else if (DefaultAssembly.IsValid())
+	{
+		return TokenFunc(DefaultAssembly);
+	}
+
+	// Cannot evaluate token function without a valid assembly
+	return FText::GetEmpty();
+}
+
+#undef LOCTEXT_NAMESPACE
diff --git a/Engine/Plugins/VirtualProduction/Shots/Source/ShotsEditor/Private/CineAssemblyNamingTokens.h b/Engine/Plugins/VirtualProduction/Shots/Source/ShotsEditor/Private/CineAssemblyNamingTokens.h
new file mode 100644
--- /dev/null
+++ b/Engine/Plugins/VirtualProduction/Shots/Source/ShotsEditor/Private/CineAssemblyNamingTokens.h
@@ -0,0 +1,56 @@
+// Copyright Epic Games, Inc. All Rights Reserved.
+
+#pragma once
+
+#include "NamingTokens.h"
+#include "CineAssembly.h"
+
+#include "CineAssemblyNamingTokens.generated.h"
+
+/** Context object that lets callers supply an assembly asset when evaluating tokens */
+UCLASS()
+class UCineAssemblyNamingTokensContext : public UObject
+{
+	GENERATED_BODY()
+
+public:
+	/** The sssembly asset to use when evaluating tokens */
+	TWeakObjectPtr<UCineAssembly> Assembly;
+};
+
+/* Naming Tokens related to Cinematic Assemblies */
+UCLASS(NotBlueprintable)
+class UCineAssemblyNamingTokens : public UNamingTokens
+{
+	GENERATED_BODY()
+
+protected:
+	// ~Begin UNamingTokens
+	virtual void OnCreateDefaultTokens(TArray<FNamingTokenData>& Tokens) override;
+	virtual void OnPreEvaluate_Implementation(const FNamingTokensEvaluationData& InEvaluationData) override;
+	virtual void OnPostEvaluate_Implementation() override;
+	// ~End UNamingTokens
+
+public:
+
+	UCineAssemblyNamingTokens();
+
+	static FString TokenNamespace;
+
+	/** Adds a metadata token with the same name as the input key which will return the result of value of that metadata if found in the Cine Assembly */
+	void AddMetadataToken(const FString& InTokenKey);
+
+	/** Sets a default assembly asset to use (useful when there is only ever one assembly needed in a given context to skip passing a context) */
+	void SetDefaultCineAssembly(UCineAssembly* InAssembly);
+
+private:
+	/** Runs the input Token evaluation function, which takes a Cine Assembly as input and outputs FText */
+	FText ExecuteTokenFunc(TFunction<FText(TWeakObjectPtr<UCineAssembly>)> TokenFunc);
+
+private:
+	/** The assembly asset to use when evaluating tokens if no context is provided */
+	TWeakObjectPtr<UCineAssembly> DefaultAssembly;
+
+	/** The current context to use when evaluating tokens */
+	TObjectPtr<UCineAssemblyNamingTokensContext> Context;
+};
diff --git a/Engine/Plugins/VirtualProduction/Shots/Source/ShotsEditor/Private/CineAssemblySchema.h b/Engine/Plugins/VirtualProduction/Shots/Source/ShotsEditor/Private/CineAssemblySchema.h
--- a/Engine/Plugins/VirtualProduction/Shots/Source/ShotsEditor/Private/CineAssemblySchema.h
+++ b/Engine/Plugins/VirtualProduction/Shots/Source/ShotsEditor/Private/CineAssemblySchema.h
@@ -83,6 +83,10 @@ public:
 	UPROPERTY(EditAnywhere, Category = "Default")
 	FString Description;
 
+	/** The default name to be use when creating assemblies from this schema */
+	UPROPERTY(EditAnywhere, Category = "Default")
+	FString DefaultAssemblyName;
+
 	/** The thumbnail image to use for this schema and assemblies built from this schema */
 	UPROPERTY(EditAnywhere, Category = "Default")
 	TObjectPtr<UTexture2D> ThumbnailImage;
diff --git a/Engine/Plugins/VirtualProduction/Shots/Source/ShotsEditor/Private/UI/CineAssembly/SCineAssemblyConfigWindow.cpp b/Engine/Plugins/VirtualProduction/Shots/Source/ShotsEditor/Private/UI/CineAssembly/SCineAssemblyConfigWindow.cpp
--- a/Engine/Plugins/VirtualProduction/Shots/Source/ShotsEditor/Private/UI/CineAssembly/SCineAssemblyConfigWindow.cpp
+++ b/Engine/Plugins/VirtualProduction/Shots/Source/ShotsEditor/Private/UI/CineAssembly/SCineAssemblyConfigWindow.cpp
@@ -2,13 +2,16 @@
 
 #include "SCineAssemblyConfigWindow.h"
 
+#include "CineAssemblyNamingTokens.h"
 #include "ContentBrowserModule.h"
 #include "Editor.h"
 #include "IContentBrowserSingleton.h"
 #include "IStructureDetailsView.h"
 #include "Modules/ModuleManager.h"
+#include "NamingTokensEngineSubsystem.h"
 #include "PropertyEditorModule.h"
 #include "ShotsEditorStyle.h"
+#include "STemplateStringEditableTextBox.h"
 #include "Widgets/SBoxPanel.h"
 #include "Widgets/Input/SButton.h"
 #include "Widgets/Input/SSegmentedControl.h"
@@ -34,6 +37,13 @@ void SCineAssemblyConfigWindow::Construct(const FArguments& InArgs, UCineAssembl
 
 	const FVector2D DefaultWindowSize = FVector2D(1400, 750);
 
+	GetOnWindowClosedEvent().AddSP(this, &SCineAssemblyConfigWindow::OnWindowClosed);
+
+	// Set the assembly being configured as the default assembly to use when evaluating naming tokens to avoid passing a context object every time we evaluate
+	UNamingTokensEngineSubsystem* NamingTokensSubsystem = GEngine->GetEngineSubsystem<UNamingTokensEngineSubsystem>();
+	UCineAssemblyNamingTokens* CineAssemblyNamingTokens = Cast<UCineAssemblyNamingTokens>(NamingTokensSubsystem->GetNamingTokens(UCineAssemblyNamingTokens::TokenNamespace));
+	CineAssemblyNamingTokens->SetDefaultCineAssembly(CineAssemblyToConfigure);
+
 	FText WindowTitle;
 	if (Mode == ECineAssemblyConfigMode::CreateNew)
 	{
@@ -92,6 +102,14 @@ void SCineAssemblyConfigWindow::Construct(const FArguments& InArgs, UCineAssembl
 		]);
 }
 
+void SCineAssemblyConfigWindow::OnWindowClosed(const TSharedRef<SWindow>& InWindow)
+{
+	// Set the assembly being configured as the default assembly to use when evaluating naming tokens to avoid passing a context object every time we evaluate
+	UNamingTokensEngineSubsystem* NamingTokensSubsystem = GEngine->GetEngineSubsystem<UNamingTokensEngineSubsystem>();
+	UCineAssemblyNamingTokens* CineAssemblyNamingTokens = Cast<UCineAssemblyNamingTokens>(NamingTokensSubsystem->GetNamingTokens(UCineAssemblyNamingTokens::TokenNamespace));
+	CineAssemblyNamingTokens->SetDefaultCineAssembly(nullptr);
+}
+
 TSharedRef<SWidget> SCineAssemblyConfigWindow::MakeCineTemplatePanel()
 {
 	FContentBrowserModule& ContentBrowserModule = FModuleManager::Get().LoadModuleChecked<FContentBrowserModule>("ContentBrowser");
@@ -148,10 +166,12 @@ void SCineAssemblyConfigWindow::OnSchemaSelected(const FAssetData& AssetData)
 	if (UCineAssemblySchema* CineAssemblySchema = Cast<UCineAssemblySchema>(AssetData.GetAsset()))
 	{
 		SelectedSchema = CineAssemblySchema;
+		CineAssemblyToConfigure->AssemblyName.Template = SelectedSchema->DefaultAssemblyName;
 	}
 	else
 	{
 		SelectedSchema = nullptr;
+		CineAssemblyToConfigure->AssemblyName.Template = TEXT("");
 	}
 
 	CineAssemblyToConfigure->SetSchema(SelectedSchema);
@@ -222,9 +242,62 @@ TSharedRef<SWidget> SCineAssemblyConfigWindow::MakeInfoPanel()
 		.HAlign(HAlign_Fill)
 		[
 			TabSwitcher.ToSharedRef()
+		]
+
+		+ SVerticalBox::Slot()
+		.AutoHeight()
+		[
+			SNew(SBorder)
+			.BorderImage(FShotsEditorStyle::Get().GetBrush("ProductionWizard.PanelBackground"))
+			.Padding(16.0f)
+				[
+					SNew(SHorizontalBox)
+
+					+ SHorizontalBox::Slot()
+						.AutoWidth()
+						.VAlign(VAlign_Center)
+						.Padding(0.0f, 0.0f, 8.0f, 0.0f)
+						[
+							SNew(STextBlock).Text(LOCTEXT("AssemblyNameField", "Assembly Name"))
+						]
+
+					+ SHorizontalBox::Slot()
+						.FillWidth(1.0f)
+						.VAlign(VAlign_Center)
+						[
+							SNew(STemplateStringEditableTextBox)
+								.Text_Lambda([this]() -> FText 
+									{
+										return FText::FromString(CineAssemblyToConfigure->AssemblyName.Template);
+									})
+								.ResolvedText_Lambda([this]() 
+									{ 
+										EvaluateTokenString(CineAssemblyToConfigure->AssemblyName);
+										return CineAssemblyToConfigure->AssemblyName.Resolved;
+									})
+								.OnTextCommitted_Lambda([this](const FText& InText, ETextCommit::Type InCommitType)
+									{
+										CineAssemblyToConfigure->Modify();
+										CineAssemblyToConfigure->AssemblyName.Template = InText.ToString();
+									})
+						]
+				]
 		];
 }
 
+void SCineAssemblyConfigWindow::EvaluateTokenString(FTemplateString& StringToEvaluate)
+{
+	FDateTime CurrentTime = FDateTime::Now();
+	if ((CurrentTime - LastTokenUpdateTime).GetSeconds() >= 1.0f)
+	{
+		UNamingTokensEngineSubsystem* NamingTokensSubsystem = GEngine->GetEngineSubsystem<UNamingTokensEngineSubsystem>();
+		FNamingTokenResultData Result = NamingTokensSubsystem->EvaluateTokenString(CineAssemblyToConfigure->AssemblyName.Template, FNamingTokenFilterArgs({ UCineAssemblyNamingTokens::TokenNamespace }));
+
+		CineAssemblyToConfigure->AssemblyName.Resolved = Result.EvaluatedText;
+		LastTokenUpdateTime = CurrentTime;
+	}
+}
+
 TSharedRef<SWidget> SCineAssemblyConfigWindow::MakeOverviewWidget()
 {
 	return SNew(SBorder)
diff --git a/Engine/Plugins/VirtualProduction/Shots/Source/ShotsEditor/Private/UI/CineAssembly/SCineAssemblyConfigWindow.h b/Engine/Plugins/VirtualProduction/Shots/Source/ShotsEditor/Private/UI/CineAssembly/SCineAssemblyConfigWindow.h
--- a/Engine/Plugins/VirtualProduction/Shots/Source/ShotsEditor/Private/UI/CineAssembly/SCineAssemblyConfigWindow.h
+++ b/Engine/Plugins/VirtualProduction/Shots/Source/ShotsEditor/Private/UI/CineAssembly/SCineAssemblyConfigWindow.h
@@ -61,6 +61,15 @@ private:
 	/** Updates the UI and CineAssembly properties based on the selected schema */
 	void OnSchemaSelected(const FAssetData& AssetData);
 
+	/** 
+	 * Evaluates the input template string with the naming tokens subsystem, and stores the result in the Resolved text.
+	 * This function is throttled to only run at a set frequency, to avoid the potential to constantly query the naming tokens subsystem.
+	 */
+	void EvaluateTokenString(FTemplateString& StringToEvaluate);
+
+	/** Callback when this window is about to be closed */
+	void OnWindowClosed(const TSharedRef<SWindow>& InWindow);
+
 private:
 	/** Switcher that controls which tab widget is currently visible */
 	TSharedPtr<SWidgetSwitcher> TabSwitcher;
@@ -79,4 +88,7 @@ private:
 
 	/** Mode that configures the UI based on the intended user interactions with the Cine Assembly asset */
 	ECineAssemblyConfigMode Mode = ECineAssemblyConfigMode::CreateNew;
+
+	/** The last time the naming tokens were updated */
+	FDateTime LastTokenUpdateTime;
 };
diff --git a/Engine/Plugins/VirtualProduction/Shots/Source/ShotsEditor/Private/UI/CineAssembly/SCineAssemblySchemaWindow.cpp b/Engine/Plugins/VirtualProduction/Shots/Source/ShotsEditor/Private/UI/CineAssembly/SCineAssemblySchemaWindow.cpp
--- a/Engine/Plugins/VirtualProduction/Shots/Source/ShotsEditor/Private/UI/CineAssembly/SCineAssemblySchemaWindow.cpp
+++ b/Engine/Plugins/VirtualProduction/Shots/Source/ShotsEditor/Private/UI/CineAssembly/SCineAssemblySchemaWindow.cpp
@@ -4,13 +4,16 @@
 
 #include "AssetRegistry/AssetRegistryModule.h"
 #include "AssetViewUtils.h"
+#include "CineAssemblyNamingTokens.h"
 #include "Framework/MultiBox/MultiBoxBuilder.h"
 #include "IStructureDetailsView.h"
 #include "Modules/ModuleManager.h"
+#include "NamingTokensEngineSubsystem.h"
 #include "PropertyEditorModule.h"
 #include "ShotsEditorStyle.h"
 #include "Styling/StyleColors.h"
 #include "Widgets/Input/SButton.h"
+#include "Widgets/Input/SEditableTextBox.h"
 #include "Widgets/Input/SSegmentedControl.h"
 #include "Widgets/Layout/SBorder.h"
 #include "Widgets/Layout/SSeparator.h"
@@ -139,7 +142,54 @@ TSharedRef<SWidget> SCineAssemblySchemaWindow::MakeContentPanel()
 			MakeHierarchyTabContent()
 		];
 
-	return TabSwitcher.ToSharedRef();
+	return SNew(SVerticalBox)
+
+		+ SVerticalBox::Slot()
+		.FillHeight(1.0f)
+		[
+			TabSwitcher.ToSharedRef()
+		]
+
+		+SVerticalBox::Slot()
+		.AutoHeight()
+		[
+			SNew(SBorder)
+				.BorderImage(FShotsEditorStyle::Get().GetBrush("ProductionWizard.PanelBackground"))
+				.Padding(16.0f)
+				[
+					SNew(SHorizontalBox)
+
+					+ SHorizontalBox::Slot()
+						.FillWidth(0.8)
+						[
+							SNullWidget::NullWidget
+						]
+
+					+ SHorizontalBox::Slot()
+						.AutoWidth()
+						.VAlign(VAlign_Center)
+						.Padding(0.0f, 0.0f, 8.0f, 0.0f)
+						[
+							SNew(STextBlock).Text(LOCTEXT("SchemaNameField", "Schema Name"))
+						]
+
+					+ SHorizontalBox::Slot()
+						.VAlign(VAlign_Center)
+						.FillWidth(0.2)
+						[
+							SNew(SEditableTextBox)
+								.Text_Lambda([this]() -> FText 
+									{ 
+										return FText::FromString(SchemaToConfigure->SchemaName); 
+									})
+								.OnTextCommitted_Lambda([this](const FText& InText, ETextCommit::Type InCommitType)
+									{
+										SchemaToConfigure->Modify();
+										SchemaToConfigure->SchemaName = InText.ToString();
+									})
+						]
+				]
+		];
 }
 
 TSharedRef<SWidget> SCineAssemblySchemaWindow::MakeButtonsPanel()
@@ -923,6 +973,15 @@ void SCineAssemblySchemaWindow::OnSchemaWindowClosed(const TSharedRef<SWindow>&
 		SchemaToConfigure->RenameAsset();
 	}
 
+	// Add each metadata key to the list of supported naming tokens
+	UNamingTokensEngineSubsystem* NamingTokensSubsystem = GEngine->GetEngineSubsystem<UNamingTokensEngineSubsystem>();
+	UCineAssemblyNamingTokens* CineAssemblyNamingTokens = Cast<UCineAssemblyNamingTokens>(NamingTokensSubsystem->GetNamingTokens(UCineAssemblyNamingTokens::TokenNamespace));
+
+	for (const FAssemblyMetadataDesc& MetadataDesc : SchemaToConfigure->AssemblyMetadata)
+	{
+		CineAssemblyNamingTokens->AddMetadataToken(MetadataDesc.Key);
+	}
+
 	// Update the schema's folder list
 	const TArray<FString> CachedFolderList = SchemaToConfigure->FoldersToCreate;
 
diff --git a/Engine/Plugins/VirtualProduction/Shots/Source/ShotsEditor/ShotsEditor.build.cs b/Engine/Plugins/VirtualProduction/Shots/Source/ShotsEditor/ShotsEditor.build.cs
--- a/Engine/Plugins/VirtualProduction/Shots/Source/ShotsEditor/ShotsEditor.build.cs
+++ b/Engine/Plugins/VirtualProduction/Shots/Source/ShotsEditor/ShotsEditor.build.cs
@@ -26,6 +26,7 @@ namespace UnrealBuildTool.Rules
 					"DesktopPlatform",
 					"DeveloperSettings",
 					"DirectoryPlaceholder",
+					"EditorWidgets",
 					"Engine",
 					"InputCore",
 					"Json",
