ShotManagement: Fix crash when user tries to move/delete root folder in schema asset tree
#rnx

diff --git a/Engine/Plugins/VirtualProduction/Shots/Source/ShotsEditor/Private/UI/CineAssembly/SCineAssemblySchemaWindow.cpp b/Engine/Plugins/VirtualProduction/Shots/Source/ShotsEditor/Private/UI/CineAssembly/SCineAssemblySchemaWindow.cpp
--- a/Engine/Plugins/VirtualProduction/Shots/Source/ShotsEditor/Private/UI/CineAssembly/SCineAssemblySchemaWindow.cpp
+++ b/Engine/Plugins/VirtualProduction/Shots/Source/ShotsEditor/Private/UI/CineAssembly/SCineAssemblySchemaWindow.cpp
@@ -712,7 +712,10 @@ FReply SCineAssemblySchemaWindow::OnTreeViewKeyDown(const FGeometry& MyGeometry,
 		if (SelectedItems.Num() == 1)
 		{
 			const TSharedPtr<FSchemaTreeItem>& SelectedItem = SelectedItems[0];
-			DeleteTreeItem(SelectedItem);
+			if (SelectedItem != RootItem)
+			{
+				DeleteTreeItem(SelectedItem);
+			}
 		}
 	}
 
@@ -1044,6 +1047,12 @@ FReply SCineAssemblySchemaWindow::OnTreeRowDragDetected(const FGeometry& MyGeome
 		if (SelectedTreeItems.Num() == 1)
 		{
 			TSharedPtr<FSchemaTreeItem> SelectedTreeItem = SelectedTreeItems[0];
+		
+			if (SelectedTreeItem == RootItem)
+			{
+				return FReply::Handled();
+			}
+			
 			TSharedRef<FSchemaAssetDragDrop> Operation = FSchemaAssetDragDrop::New(FPaths::GetPathLeaf(SelectedTreeItem->Path));
 			Operation->SourceTreeItem = SelectedTreeItem;
 
