Shots Management: When an assembly is selected to render in MRQ, the movie pipeline job's map will be populated based on the assembly.

Add static OnSequenceSetEvent() to UMoviePipelineExecutorJob which fires when SetSequence is called.

diff --git a/Engine/Plugins/MovieScene/MovieRenderPipeline/Source/MovieRenderPipelineCore/Private/MoviePipelineQueue.cpp b/Engine/Plugins/MovieScene/MovieRenderPipeline/Source/MovieRenderPipelineCore/Private/MoviePipelineQueue.cpp
--- a/Engine/Plugins/MovieScene/MovieRenderPipeline/Source/MovieRenderPipelineCore/Private/MoviePipelineQueue.cpp
+++ b/Engine/Plugins/MovieScene/MovieRenderPipeline/Source/MovieRenderPipelineCore/Private/MoviePipelineQueue.cpp
@@ -291,6 +291,9 @@ void UMoviePipelineExecutorJob::SetSequence(FSoftObjectPath InSequence)
 		return;
 	}
 
+	// Allow external tools to set job properties based on the current sequence.
+	OnSequenceSetEvent().Broadcast(this, LoadedSequence);
+
 	bool bShotsChanged = false;
 	UMoviePipelineBlueprintLibrary::UpdateJobShotListFromSequence(LoadedSequence, this, bShotsChanged);
 	
@@ -300,6 +303,12 @@ void UMoviePipelineExecutorJob::SetSequence(FSoftObjectPath InSequence)
 	}
 }
 
+UMoviePipelineExecutorJob::FOnSequenceSet& UMoviePipelineExecutorJob::OnSequenceSetEvent()
+{
+	static FOnSequenceSet OnSequenceSetEvent;
+	return OnSequenceSetEvent;
+}
+
 UMovieJobVariableAssignmentContainer* UMoviePipelineExecutorJob::GetOrCreateJobVariableAssignmentsForGraph(const UMovieGraphConfig* InGraph)
 {
 	if (InGraph)
diff --git a/Engine/Plugins/MovieScene/MovieRenderPipeline/Source/MovieRenderPipelineCore/Public/MoviePipelineQueue.h b/Engine/Plugins/MovieScene/MovieRenderPipeline/Source/MovieRenderPipelineCore/Public/MoviePipelineQueue.h
--- a/Engine/Plugins/MovieScene/MovieRenderPipeline/Source/MovieRenderPipelineCore/Public/MoviePipelineQueue.h
+++ b/Engine/Plugins/MovieScene/MovieRenderPipeline/Source/MovieRenderPipelineCore/Public/MoviePipelineQueue.h
@@ -502,6 +502,11 @@ public:
 	UFUNCTION(BlueprintSetter, Category = "Movie Render Pipeline")
 	void SetSequence(FSoftObjectPath InSequence);
 
+	DECLARE_MULTICAST_DELEGATE_TwoParams(FOnSequenceSet, UMoviePipelineExecutorJob* Job, ULevelSequence* Sequence);
+
+	/** Retrieve multicast delegate for when a sequence is set. */
+	static FOnSequenceSet& OnSequenceSetEvent();
+	
 	/**
 	* This method will return you the object which contains variable overrides for the Primary Graph assigned to this job. You need to provide
 	* a pointer to the exact graph you want (as the Primary Graph may contain sub-graphs, and those sub-graphs can have their own variables),
diff --git a/Engine/Plugins/VirtualProduction/Shots/Shots.uplugin b/Engine/Plugins/VirtualProduction/Shots/Shots.uplugin
--- a/Engine/Plugins/VirtualProduction/Shots/Shots.uplugin
+++ b/Engine/Plugins/VirtualProduction/Shots/Shots.uplugin
@@ -33,6 +33,10 @@
 			"Name": "DirectoryPlaceholder",
 			"Enabled": true
 		},
+		{
+			"Name": "MovieRenderPipeline",
+			"Enabled": true
+		},
 		{
 			"Name": "NamingTokens",
 			"Enabled": true
diff --git a/Engine/Plugins/VirtualProduction/Shots/Source/ShotsEditor/Private/ShotsEditorModule.cpp b/Engine/Plugins/VirtualProduction/Shots/Source/ShotsEditor/Private/ShotsEditorModule.cpp
--- a/Engine/Plugins/VirtualProduction/Shots/Source/ShotsEditor/Private/ShotsEditorModule.cpp
+++ b/Engine/Plugins/VirtualProduction/Shots/Source/ShotsEditor/Private/ShotsEditorModule.cpp
@@ -11,6 +11,7 @@
 #include "Framework/Application/SlateApplication.h"
 #include "Framework/Docking/TabManager.h"
 #include "Interfaces/IMainFrameModule.h"
+#include "MoviePipelineQueue.h"
 #include "ProductionSettings.h"
 #include "ProductionSettingsCustomization.h"
 #include "PropertyEditorModule.h"
@@ -39,6 +40,12 @@ public:
 private:
 	TSharedRef<SDockTab> SpawnProductionWizard(const FSpawnTabArgs& SpawnTabArgs);
 	void RegisterMenus();
+
+	/** Called whenever a job has its sequence set. */
+	void OnSequenceSet(UMoviePipelineExecutorJob* Job, ULevelSequence* Sequence);
+
+private:
+	FDelegateHandle OnSequenceSetHandle;
 };
 
 void FShotsEditorModule::StartupModule()
@@ -70,6 +77,8 @@ void FShotsEditorModule::StartupModule()
 		.SetIcon(FSlateIcon(FAppStyle::GetAppStyleSetName(), "ProjectSettings.TabIcon"));
 
 	RegisterMenus();
+
+	OnSequenceSetHandle = UMoviePipelineExecutorJob::OnSequenceSetEvent().AddRaw(this, &FShotsEditorModule::OnSequenceSet);
 }
 
 void FShotsEditorModule::ShutdownModule()
@@ -78,6 +87,8 @@ void FShotsEditorModule::ShutdownModule()
 	PropertyModule.UnregisterCustomClassLayout(UProductionSettings::StaticClass()->GetFName());
 
 	FGlobalTabmanager::Get()->UnregisterNomadTabSpawner(ProductionWizardTabName);
+
+	UMoviePipelineExecutorJob::OnSequenceSetEvent().Remove(OnSequenceSetHandle);
 }
 
 TSharedRef<SDockTab> FShotsEditorModule::SpawnProductionWizard(const FSpawnTabArgs& SpawnTabArgs)
@@ -144,6 +155,21 @@ void FShotsEditorModule::RegisterMenus()
 	}
 }
 
+void FShotsEditorModule::OnSequenceSet(UMoviePipelineExecutorJob* Job, ULevelSequence* Sequence)
+{
+	if (Job)
+	{
+		if (const UCineAssembly* CineAssembly = Cast<UCineAssembly>(Sequence))
+		{
+			// For our cine assemblies, we want to always use the map associated with the sequence.
+			if (CineAssembly->Level.IsValid())
+			{
+				Job->Map = CineAssembly->Level;
+			}
+		}
+	}
+}
+
 IMPLEMENT_MODULE(FShotsEditorModule, ShotsEditor);
 
 #undef LOCTEXT_NAMESPACE
diff --git a/Engine/Plugins/VirtualProduction/Shots/Source/ShotsEditor/ShotsEditor.build.cs b/Engine/Plugins/VirtualProduction/Shots/Source/ShotsEditor/ShotsEditor.build.cs
--- a/Engine/Plugins/VirtualProduction/Shots/Source/ShotsEditor/ShotsEditor.build.cs
+++ b/Engine/Plugins/VirtualProduction/Shots/Source/ShotsEditor/ShotsEditor.build.cs
@@ -33,6 +33,7 @@ namespace UnrealBuildTool.Rules
 					"JsonUtilities",
 					"LevelSequence",
 					"LevelSequenceEditor",
+					"MovieRenderPipelineCore",
 					"MovieScene",
 					"MovieSceneTools",
 					"NamingTokens",
