ShotManagement: Add additional settings for Cinematic Assemblies to Take Recorder, including a Schema that should be used when recording a new Cine Assembly.
When recording a Cine Assembly, the schema will be set and the correct subsequences will be created before the recording begins. These subsequence tracks will grow in length during the recording alongside the other recorded tracks, and will have a matching frame range when recording has completed.
#rnx

diff --git a/Engine/Plugins/VirtualProduction/Shots/Shots.uplugin b/Engine/Plugins/VirtualProduction/Shots/Shots.uplugin
--- a/Engine/Plugins/VirtualProduction/Shots/Shots.uplugin
+++ b/Engine/Plugins/VirtualProduction/Shots/Shots.uplugin
@@ -40,6 +40,10 @@
 		{
 			"Name": "NamingTokens",
 			"Enabled": true
+		},
+		{
+			"Name": "Takes",
+			"Enabled": true
 		}
 	]
 }
\ No newline at end of file
diff --git a/Engine/Plugins/VirtualProduction/Shots/Source/ShotsEditor/Private/CineAssembly.cpp b/Engine/Plugins/VirtualProduction/Shots/Source/ShotsEditor/Private/CineAssembly.cpp
--- a/Engine/Plugins/VirtualProduction/Shots/Source/ShotsEditor/Private/CineAssembly.cpp
+++ b/Engine/Plugins/VirtualProduction/Shots/Source/ShotsEditor/Private/CineAssembly.cpp
@@ -2,11 +2,18 @@
 
 #include "CineAssembly.h"
 
+#include "AssetRegistry/AssetRegistryModule.h"
+#include "AssetToolsModule.h"
+#include "CineAssemblyFactory.h"
 #include "CineAssemblyNamingTokens.h"
+#include "HAL/FileManager.h"
+#include "MovieScene.h"
 #include "NamingTokensEngineSubsystem.h"
+#include "Sections/MovieSceneSubSection.h"
 #include "Serialization/JsonReader.h"
 #include "Serialization/JsonSerializer.h"
 #include "Serialization/JsonWriter.h"
+#include "Tracks/MovieSceneSubTrack.h"
 #include "UObject/AssetRegistryTagsContext.h"
 
 #define LOCTEXT_NAMESPACE "CineAssembly"
@@ -24,6 +31,14 @@ const UCineAssemblySchema* UCineAssembly::GetSchema() const
 }
 
 void UCineAssembly::SetSchema(UCineAssemblySchema* InSchema)
+{
+	if (BaseSchema == nullptr)
+	{
+		BaseSchema = InSchema;
+	}
+}
+
+void UCineAssembly::ChangeSchema(UCineAssemblySchema* InSchema)
 {
 	// Remove all metadata associated with the old schema before changing it
 	if (BaseSchema)
@@ -60,6 +75,67 @@ void UCineAssembly::SetSchema(UCineAssemblySchema* InSchema)
 	}
 }
 
+void UCineAssembly::CreateSubAssemblies()
+{
+	if (!BaseSchema || !SubAssemblies.IsEmpty())
+	{
+		return;
+	}
+
+	// Get the path where the top-level assembly will be created so we can create other assets relative to it
+	const FAssetData AssemblyAssetData = FAssetData(this);
+	const FString PackagePath = AssemblyAssetData.PackagePath.ToString();
+
+	// Create the default folders for this assembly, based on the schema
+	for (const FString& FolderPath : BaseSchema->FoldersToCreate)
+	{
+		const FString PathToCreate = PackagePath / FolderPath;
+		const FString RelativeFilePath = FPackageName::LongPackageNameToFilename(PathToCreate);
+		const FString AbsoluteFilePath = FPaths::ConvertRelativePathToFull(RelativeFilePath);
+
+		// Create the directory on disk, then add its path to the asset registry so it appears in Content Browser
+		if (!IFileManager::Get().DirectoryExists(*AbsoluteFilePath))
+		{
+			constexpr bool bCreateParentFoldersIfMissing = true;
+			if (IFileManager::Get().MakeDirectory(*AbsoluteFilePath, bCreateParentFoldersIfMissing))
+			{
+				FAssetRegistryModule& AssetRegistryModule = FModuleManager::GetModuleChecked<FAssetRegistryModule>("AssetRegistry");
+				AssetRegistryModule.Get().AddPath(PathToCreate);
+			}
+		}
+	}
+
+	IAssetTools& AssetTools = FModuleManager::GetModuleChecked<FAssetToolsModule>("AssetTools").Get();
+
+	// Create a new CineAssembly for each subsequence, set its playback range to match the parent sequence, and add it to the subsequence track
+	for (const FString& Subsequence : BaseSchema->SubsequencesToCreate)
+	{
+		// Add a subsequence track to the assembly's sequence
+		UMovieSceneSubTrack* SubTrack = Cast<UMovieSceneSubTrack>(MovieScene->AddTrack(UMovieSceneSubTrack::StaticClass()));
+		SubTrack->SetDisplayName(FText::FromString(FPaths::GetPathLeaf(Subsequence)));
+
+		// Before creating each subassembly, sanity check that each one will actually have a unique name (in case there are duplicates in the schema description)
+		FString UniquePackageName;
+		FString UniqueAssetName;
+		AssetTools.CreateUniqueAssetName(PackagePath / Subsequence, TEXT(""), UniquePackageName, UniqueAssetName);
+
+		const FString SubAssemblyPath = FPaths::GetPath(UniquePackageName);
+		UCineAssembly* SubAssembly = Cast<UCineAssembly>(AssetTools.CreateAsset(UniqueAssetName, SubAssemblyPath, UCineAssembly::StaticClass(), NewObject<UCineAssemblyFactory>()));
+		
+		if (SubAssembly)
+		{
+			SubAssembly->Level = Level;
+
+			const TRange<FFrameNumber> PlaybackRange = SubAssembly->GetMovieScene()->GetPlaybackRange();
+			const FFrameNumber StartFrame = PlaybackRange.GetLowerBoundValue();
+			const int32 Duration = PlaybackRange.Size<FFrameNumber>().Value;
+
+			UMovieSceneSubSection* SubSection = SubTrack->AddSequence(SubAssembly, StartFrame, Duration);
+			SubAssemblies.Add(SubSection);
+		}
+	}
+}
+
 void UCineAssembly::GetAssetRegistryTags(FAssetRegistryTagsContext Context) const
 {
 	Super::GetAssetRegistryTags(Context);
diff --git a/Engine/Plugins/VirtualProduction/Shots/Source/ShotsEditor/Private/CineAssembly.h b/Engine/Plugins/VirtualProduction/Shots/Source/ShotsEditor/Private/CineAssembly.h
--- a/Engine/Plugins/VirtualProduction/Shots/Source/ShotsEditor/Private/CineAssembly.h
+++ b/Engine/Plugins/VirtualProduction/Shots/Source/ShotsEditor/Private/CineAssembly.h
@@ -10,6 +10,8 @@
 
 #include "CineAssembly.generated.h"
 
+class UMovieSceneSubSection;
+
 /** 
  * A cinematic building block that associates a level sequence with a level
  */
@@ -33,6 +35,12 @@ public:
 	/** Get the base schema for this assembly */
 	const UCineAssemblySchema* GetSchema() const;
 
+	/** Set the base schema for this assembly, only if one is not already set */
+	void SetSchema(UCineAssemblySchema* InSchema);
+
+	/** Creates one or more subsequence assets, parented to this assembly, based on the schema */
+	void CreateSubAssemblies();
+
 	/** Get all of the metadata for this assembly as a formatted JSON string */
 	UFUNCTION(BlueprintPure, Category = "Shot Management")
 	FString GetFullMetadataString() const;
@@ -63,7 +71,7 @@ public:
 
 private:
 	/** Sets the base schema for this assembly and re-initializes the metadata inherited from the schema */
-	void SetSchema(UCineAssemblySchema* InSchema);
+	void ChangeSchema(UCineAssemblySchema* InSchema);
 
 	/** Update the underlying json object whenever keys/values in the instance metadata map are added/removed/modified */
 	void UpdateInstanceMetadata();
@@ -88,6 +96,10 @@ public:
 	UPROPERTY()
 	FString AssemblyNote;
 
+	/** Array of Subsequence Sections created based on the schema */
+	UPROPERTY()
+	TArray<TObjectPtr<UMovieSceneSubSection>> SubAssemblies;
+
 	/** The asset registry tag that contains the assembly type information */
 	static const FName AssetRegistryTag_AssemblyType;
 
diff --git a/Engine/Plugins/VirtualProduction/Shots/Source/ShotsEditor/Private/CineAssemblyFactory.cpp b/Engine/Plugins/VirtualProduction/Shots/Source/ShotsEditor/Private/CineAssemblyFactory.cpp
--- a/Engine/Plugins/VirtualProduction/Shots/Source/ShotsEditor/Private/CineAssemblyFactory.cpp
+++ b/Engine/Plugins/VirtualProduction/Shots/Source/ShotsEditor/Private/CineAssemblyFactory.cpp
@@ -2,11 +2,9 @@
 
 #include "CineAssemblyFactory.h"
 
-#include "AssetRegistry/AssetRegistryModule.h"
 #include "AssetToolsModule.h"
 #include "CineAssemblyNamingTokens.h"
 #include "Framework/Application/SlateApplication.h"
-#include "HAL/FileManager.h"
 #include "Interfaces/IMainFrameModule.h"
 #include "Modules/ModuleManager.h"
 #include "MovieScene.h"
@@ -59,61 +57,7 @@ UObject* UCineAssemblyFactory::FactoryCreateNew(UClass* Class, UObject* InParent
 	FFrameRate TickResolution = CineAssemblyToConfigure->GetMovieScene()->GetTickResolution();
 	CineAssemblyToConfigure->GetMovieScene()->SetPlaybackRange((MovieSceneToolsProjectSettings->DefaultStartTime * TickResolution).FloorToFrame(), (MovieSceneToolsProjectSettings->DefaultDuration * TickResolution).FloorToFrame().Value);
 
-	if (const UCineAssemblySchema* Schema = CineAssemblyToConfigure->GetSchema())
-	{
-		// Get the path where the top-level assembly will be created so we can create other assets relative to it
-		const FAssetData AssemblyAssetData = FAssetData(CineAssemblyToConfigure);
-		const FString PackagePath = AssemblyAssetData.PackagePath.ToString();
-
-		// Create the default folders for this assembly, based on the schema
-		for (const FString& FolderPath : Schema->FoldersToCreate)
-		{
-			const FString PathToCreate = PackagePath / FolderPath;
-			const FString RelativeFilePath = FPackageName::LongPackageNameToFilename(PathToCreate);
-			const FString AbsoluteFilePath = FPaths::ConvertRelativePathToFull(RelativeFilePath);
-
-			// Create the directory on disk, then add its path to the asset registry so it appears in Content Browser
-			if (!IFileManager::Get().DirectoryExists(*AbsoluteFilePath))
-			{
-				constexpr bool bCreateParentFoldersIfMissing = true;
-				if (IFileManager::Get().MakeDirectory(*AbsoluteFilePath, bCreateParentFoldersIfMissing))
-				{
-					FAssetRegistryModule& AssetRegistryModule = FModuleManager::GetModuleChecked<FAssetRegistryModule>("AssetRegistry");
-					AssetRegistryModule.Get().AddPath(PathToCreate);
-				}
-			}
-		}
-
-		// Create the subsequence assets associated with this assembly
-		if (Schema->SubsequencesToCreate.Num() > 0)
-		{
-			// Add a subsequence track to the assembly's sequence
-			UMovieScene* MovieScene = CineAssemblyToConfigure->GetMovieScene();
-			UMovieSceneSubTrack* SubTrack = Cast<UMovieSceneSubTrack>(MovieScene->AddTrack(UMovieSceneSubTrack::StaticClass()));
-
-			IAssetTools& AssetTools = FModuleManager::GetModuleChecked<FAssetToolsModule>("AssetTools").Get();
-
-			// Create a new CineAssembly for each subsequence, set its playback range to match the parent sequence, and add it to the subsequence track
-			for (const FString& Subsequence : Schema->SubsequencesToCreate)
-			{
-				// Before creating each subassembly, sanity check that each one will actually have a unique name (in case there are duplicates in the schema description)
-				FString UniquePackageName;
-				FString UniqueAssetName;
-				AssetTools.CreateUniqueAssetName(PackagePath / Subsequence, TEXT(""), UniquePackageName, UniqueAssetName);
-
-				const FString SubAssemblyPath = FPaths::GetPath(UniquePackageName);
-				UCineAssembly* SubAssembly = Cast<UCineAssembly>(AssetTools.CreateAsset(UniqueAssetName, SubAssemblyPath, UCineAssembly::StaticClass(), NewObject<UCineAssemblyFactory>()));
-
-				if (SubAssembly)
-				{
-					SubAssembly->Level = CineAssemblyToConfigure->Level;
-
-					const TRange<FFrameNumber> PlaybackRange = SubAssembly->GetMovieScene()->GetPlaybackRange();
-					SubTrack->AddSequence(SubAssembly, PlaybackRange.GetLowerBoundValue(), PlaybackRange.GetUpperBoundValue().Value - PlaybackRange.GetLowerBoundValue().Value);
-				}
-			}
-		}
-	}
+	CineAssemblyToConfigure->CreateSubAssemblies();
 
 	return CineAssemblyToConfigure;
 }
diff --git a/Engine/Plugins/VirtualProduction/Shots/Source/ShotsEditor/Private/ShotsEditorModule.cpp b/Engine/Plugins/VirtualProduction/Shots/Source/ShotsEditor/Private/ShotsEditorModule.cpp
--- a/Engine/Plugins/VirtualProduction/Shots/Source/ShotsEditor/Private/ShotsEditorModule.cpp
+++ b/Engine/Plugins/VirtualProduction/Shots/Source/ShotsEditor/Private/ShotsEditorModule.cpp
@@ -23,6 +23,7 @@
 #include "Widgets/Docking/SDockTab.h"
 #include "WorkspaceMenuStructure.h"
 #include "WorkspaceMenuStructureModule.h"
+#include "TakeRecorder/CineAssemblyTakeRecorderIntegration.h"
 
 #define LOCTEXT_NAMESPACE "ShotsEditorModule"
 
@@ -38,6 +39,9 @@ public:
 	//~ End IModuleInterface interface
 
 private:
+	/** Manages all integration of Cinematic Assembly Tools with Take Recorder. */
+	TUniquePtr<FCineAssemblyTakeRecorderIntegration> TakeRecorderIntegration;
+	
 	TSharedRef<SDockTab> SpawnProductionWizard(const FSpawnTabArgs& SpawnTabArgs);
 	void RegisterMenus();
 
@@ -78,11 +82,19 @@ void FShotsEditorModule::StartupModule()
 
 	RegisterMenus();
 
+	FCoreDelegates::OnPostEngineInit.AddLambda([this]() 
+		{
+			TakeRecorderIntegration = MakeUnique<FCineAssemblyTakeRecorderIntegration>();
+		});
+
 	OnSequenceSetHandle = UMoviePipelineExecutorJob::OnSequenceSetEvent().AddRaw(this, &FShotsEditorModule::OnSequenceSet);
 }
 
 void FShotsEditorModule::ShutdownModule()
 {
+	FCoreDelegates::OnPostEngineInit.RemoveAll(this);
+	TakeRecorderIntegration.Reset();
+
 	FPropertyEditorModule& PropertyModule = FModuleManager::GetModuleChecked<FPropertyEditorModule>("PropertyEditor");
 	PropertyModule.UnregisterCustomClassLayout(UProductionSettings::StaticClass()->GetFName());
 
diff --git a/Engine/Plugins/VirtualProduction/Shots/Source/ShotsEditor/Private/TakeRecorder/CineAssemblyTakeRecorderIntegration.cpp b/Engine/Plugins/VirtualProduction/Shots/Source/ShotsEditor/Private/TakeRecorder/CineAssemblyTakeRecorderIntegration.cpp
new file mode 100644
--- /dev/null
+++ b/Engine/Plugins/VirtualProduction/Shots/Source/ShotsEditor/Private/TakeRecorder/CineAssemblyTakeRecorderIntegration.cpp
@@ -0,0 +1,130 @@
+﻿// Copyright Epic Games, Inc. All Rights Reserved.
+
+#include "CineAssemblyTakeRecorderIntegration.h"
+
+#include "CineAssembly.h"
+#include "CineAssemblySchema.h"
+#include "CineAssemblyTakeRecorderSettings.h"
+#include "ITakeRecorderModule.h"
+#include "Recorder/TakeRecorder.h"
+#include "Sections/MovieSceneSubSection.h"
+#include "TakesUtils.h"
+
+FCineAssemblyTakeRecorderIntegration::FCineAssemblyTakeRecorderIntegration()
+{
+	UCineAssemblyTakeRecorderSettings* Settings = GetMutableDefault<UCineAssemblyTakeRecorderSettings>();
+
+	ITakeRecorderModule& TakeRecorderModule = ITakeRecorderModule::Get();
+	TakeRecorderModule.RegisterSettingsObject(Settings);
+
+	UTakeRecorder::OnRecordingInitialized().AddRaw(this, &FCineAssemblyTakeRecorderIntegration::OnRecordingInitialized);
+}
+
+FCineAssemblyTakeRecorderIntegration::~FCineAssemblyTakeRecorderIntegration()
+{
+	UTakeRecorder::OnRecordingInitialized().RemoveAll(this);
+}
+
+void FCineAssemblyTakeRecorderIntegration::OnRecordingInitialized(UTakeRecorder* TakeRecorder)
+{
+	if (TakeRecorder)
+	{
+		TakeRecorder->OnRecordingStarted().AddRaw(this, &FCineAssemblyTakeRecorderIntegration::OnRecordingStarted);
+		TakeRecorder->OnTickRecording().AddRaw(this, &FCineAssemblyTakeRecorderIntegration::OnTickRecording);
+		TakeRecorder->OnRecordingStopped().AddRaw(this, &FCineAssemblyTakeRecorderIntegration::OnRecordingStopped);
+
+		if (UCineAssembly* Assembly = Cast<UCineAssembly>(TakeRecorder->GetSequence()))
+		{
+			Assembly->Level = FSoftObjectPath(TakesUtils::DiscoverSourceWorld());
+
+			// Reset the assembly's playback range, since it was initialized to some default value when it was created, but we do not know what its final range will be until we finish recording.
+			if (UMovieScene* MovieScene = Assembly->GetMovieScene())
+			{
+				const FFrameNumber StartFrame = MovieScene->GetPlaybackRange().GetLowerBoundValue();
+				MovieScene->SetPlaybackRange(TRange<FFrameNumber>(StartFrame, StartFrame));
+			}
+
+			UCineAssemblyTakeRecorderSettings* Settings = GetMutableDefault<UCineAssemblyTakeRecorderSettings>();
+			if (!Settings->AssemblySchema.IsNull())
+			{
+				UCineAssemblySchema* Schema = Settings->AssemblySchema.LoadSynchronous();
+				Assembly->SetSchema(Schema);
+			}
+		}
+	}
+}
+
+void FCineAssemblyTakeRecorderIntegration::OnRecordingStarted(UTakeRecorder* TakeRecorder)
+{
+	if (TakeRecorder)
+	{
+		if (UCineAssembly* Assembly = Cast<UCineAssembly>(TakeRecorder->GetSequence()))
+		{
+			Assembly->CreateSubAssemblies();
+
+			// Reset the playback range of each subsection and its underlying sequence's movie scene.
+			// These were initialized to default ranges when the assets were created, but when recording, we do not want to use the default range.
+			for (UMovieSceneSubSection* SubSection : Assembly->SubAssemblies)
+			{
+				const FFrameNumber StartTime = SubSection->GetRange().GetLowerBoundValue();
+				SubSection->SetRange(TRange<FFrameNumber>(StartTime, StartTime));
+
+				if (UMovieSceneSequence* SubSequence = SubSection->GetSequence())
+				{
+					if (UMovieScene* MovieScene = SubSequence->GetMovieScene())
+					{
+						MovieScene->SetTickResolutionDirectly(Assembly->GetMovieScene()->GetTickResolution());
+						MovieScene->SetDisplayRate(Assembly->GetMovieScene()->GetDisplayRate());
+
+						MovieScene->SetPlaybackRange(TRange<FFrameNumber>(0, 0));
+					}
+				}
+			}
+		}
+	}
+}
+
+void FCineAssemblyTakeRecorderIntegration::OnTickRecording(UTakeRecorder* TakeRecorder, const FQualifiedFrameTime& CurrentFrameTime)
+{
+	if (TakeRecorder)
+	{
+		if (UCineAssembly* Assembly = Cast<UCineAssembly>(TakeRecorder->GetSequence()))
+		{
+			FFrameNumber EndFrame = CurrentFrameTime.Time.CeilToFrame();
+
+			// Expand the the frame range of each subsequence (this causes the tracks to grow longer during recording)
+			for (UMovieSceneSubSection* SubSection : Assembly->SubAssemblies)
+			{
+				SubSection->ExpandToFrame(EndFrame);
+			}
+		}
+	}
+}
+
+void FCineAssemblyTakeRecorderIntegration::OnRecordingStopped(UTakeRecorder* TakeRecorder)
+{
+	if (TakeRecorder)
+	{
+		if (UCineAssembly* Assembly = Cast<UCineAssembly>(TakeRecorder->GetSequence()))
+		{
+			// Finalize the frame range for each subsequence, lock them, and save them
+			for (UMovieSceneSubSection* SubSection : Assembly->SubAssemblies)
+			{
+				if (UMovieSceneSequence* SubSequence = SubSection->GetSequence())
+				{
+					if (UMovieScene* MovieScene = SubSequence->GetMovieScene())
+					{
+						MovieScene->SetPlaybackRange(SubSection->GetRange());
+						MovieScene->SetReadOnly(true);
+					}
+
+					TakesUtils::SaveAsset(SubSequence);
+				}
+			}
+		}
+
+		TakeRecorder->OnRecordingStarted().RemoveAll(this);
+		TakeRecorder->OnTickRecording().RemoveAll(this);
+		TakeRecorder->OnRecordingStopped().RemoveAll(this);
+	}
+}
diff --git a/Engine/Plugins/VirtualProduction/Shots/Source/ShotsEditor/Private/TakeRecorder/CineAssemblyTakeRecorderIntegration.h b/Engine/Plugins/VirtualProduction/Shots/Source/ShotsEditor/Private/TakeRecorder/CineAssemblyTakeRecorderIntegration.h
new file mode 100644
--- /dev/null
+++ b/Engine/Plugins/VirtualProduction/Shots/Source/ShotsEditor/Private/TakeRecorder/CineAssemblyTakeRecorderIntegration.h
@@ -0,0 +1,23 @@
+﻿// Copyright Epic Games, Inc. All Rights Reserved.
+
+#pragma once
+
+#include "Templates/UnrealTemplate.h"
+
+class UTakeRecorder;
+struct FQualifiedFrameTime;
+
+/** Manages Cine Assembly Tools interactions with Take Recorder */
+class FCineAssemblyTakeRecorderIntegration : public FNoncopyable
+{
+	public:
+		FCineAssemblyTakeRecorderIntegration();
+		~FCineAssemblyTakeRecorderIntegration();
+
+	private:
+		void OnRecordingInitialized(UTakeRecorder* TakeRecorder);
+
+		void OnRecordingStarted(UTakeRecorder* TakeRecorder);
+		void OnTickRecording(UTakeRecorder* TakeRecorder, const FQualifiedFrameTime& CurrentFrameTime);
+		void OnRecordingStopped(UTakeRecorder* TakeRecorder);
+};
diff --git a/Engine/Plugins/VirtualProduction/Shots/Source/ShotsEditor/Private/TakeRecorder/CineAssemblyTakeRecorderSettings.cpp b/Engine/Plugins/VirtualProduction/Shots/Source/ShotsEditor/Private/TakeRecorder/CineAssemblyTakeRecorderSettings.cpp
new file mode 100644
--- /dev/null
+++ b/Engine/Plugins/VirtualProduction/Shots/Source/ShotsEditor/Private/TakeRecorder/CineAssemblyTakeRecorderSettings.cpp
@@ -0,0 +1,9 @@
+// Copyright Epic Games, Inc. All Rights Reserved.
+
+#include "CineAssemblyTakeRecorderSettings.h"
+
+void UCineAssemblyTakeRecorderSettings::PostEditChangeProperty(FPropertyChangedEvent& PropertyChangedEvent)
+{
+	UObject::PostEditChangeProperty(PropertyChangedEvent);
+	SaveConfig();
+}
diff --git a/Engine/Plugins/VirtualProduction/Shots/Source/ShotsEditor/Private/TakeRecorder/CineAssemblyTakeRecorderSettings.h b/Engine/Plugins/VirtualProduction/Shots/Source/ShotsEditor/Private/TakeRecorder/CineAssemblyTakeRecorderSettings.h
new file mode 100644
--- /dev/null
+++ b/Engine/Plugins/VirtualProduction/Shots/Source/ShotsEditor/Private/TakeRecorder/CineAssemblyTakeRecorderSettings.h
@@ -0,0 +1,24 @@
+// Copyright Epic Games, Inc. All Rights Reserved.
+
+#pragma once
+
+#include "UObject/Object.h"
+#include "CineAssemblySchema.h"
+
+#include "CineAssemblyTakeRecorderSettings.generated.h"
+
+UCLASS(Config=EditorPerProjectUserSettings, DisplayName = "Take Recorder")
+class UCineAssemblyTakeRecorderSettings : public UObject
+{
+	GENERATED_BODY()
+
+public:
+	//~ Begin UObject Interface
+	virtual void PostEditChangeProperty(FPropertyChangedEvent& PropertyChangedEvent) override;
+	//~ End UObject Interface
+
+public:
+	/** The Cine Assembly Schema to use as the base for the recorded Assembly */
+	UPROPERTY(Config, EditAnywhere, Category = "Cine Assembly")
+	TSoftObjectPtr<UCineAssemblySchema> AssemblySchema;
+};
diff --git a/Engine/Plugins/VirtualProduction/Shots/Source/ShotsEditor/Private/UI/CineAssembly/SCineAssemblyConfigWindow.cpp b/Engine/Plugins/VirtualProduction/Shots/Source/ShotsEditor/Private/UI/CineAssembly/SCineAssemblyConfigWindow.cpp
--- a/Engine/Plugins/VirtualProduction/Shots/Source/ShotsEditor/Private/UI/CineAssembly/SCineAssemblyConfigWindow.cpp
+++ b/Engine/Plugins/VirtualProduction/Shots/Source/ShotsEditor/Private/UI/CineAssembly/SCineAssemblyConfigWindow.cpp
@@ -174,7 +174,7 @@ void SCineAssemblyConfigWindow::OnSchemaSelected(const FAssetData& AssetData)
 		CineAssemblyToConfigure->AssemblyName.Template = TEXT("");
 	}
 
-	CineAssemblyToConfigure->SetSchema(SelectedSchema);
+	CineAssemblyToConfigure->ChangeSchema(SelectedSchema);
 
 	// The details view needs to be redrawn to show the new metadata fields from the selected schema
 	DetailsView->ForceRefresh();
diff --git a/Engine/Plugins/VirtualProduction/Shots/Source/ShotsEditor/ShotsEditor.build.cs b/Engine/Plugins/VirtualProduction/Shots/Source/ShotsEditor/ShotsEditor.build.cs
--- a/Engine/Plugins/VirtualProduction/Shots/Source/ShotsEditor/ShotsEditor.build.cs
+++ b/Engine/Plugins/VirtualProduction/Shots/Source/ShotsEditor/ShotsEditor.build.cs
@@ -43,9 +43,11 @@ namespace UnrealBuildTool.Rules
 					"SlateCore",
 					"SourceControl",
 					"TimeManagement",
+					"TakeRecorder",
+					"TakesCore",
 					"ToolMenus",
 					"UnrealEd",
-					"WorkspaceMenuStructure"
+					"WorkspaceMenuStructure", 
 				}
 			);
 		}
diff --git a/Engine/Plugins/VirtualProduction/Takes/Source/TakeRecorder/Private/Recorder/TakeRecorder.cpp b/Engine/Plugins/VirtualProduction/Takes/Source/TakeRecorder/Private/Recorder/TakeRecorder.cpp
--- a/Engine/Plugins/VirtualProduction/Takes/Source/TakeRecorder/Private/Recorder/TakeRecorder.cpp
+++ b/Engine/Plugins/VirtualProduction/Takes/Source/TakeRecorder/Private/Recorder/TakeRecorder.cpp
@@ -823,6 +823,8 @@ void UTakeRecorder::InternalTick(float DeltaTime)
 		}
 	}
 
+	OnTickRecordingEvent.Broadcast(this, RecordTime);
+
 	if (StopRecordingFrame.IsSet() && CurrentFrameTime.FrameNumber >= StopRecordingFrame.GetValue())
 	{
 		Stop();
@@ -1313,6 +1315,11 @@ FOnStartPlayFrameModified& UTakeRecorder::OnStartPlayFrameModified()
 	return OnFrameModifiedEvent;
 }
 
+FOnTickRecording& UTakeRecorder::OnTickRecording()
+{
+	return OnTickRecordingEvent;
+}
+
 void UTakeRecorder::HandlePIE(bool bIsSimulating)
 {
 	ULevelSequence* FinishedAsset = GetSequence();
diff --git a/Engine/Plugins/VirtualProduction/Takes/Source/TakeRecorder/Public/ITakeRecorderModule.h b/Engine/Plugins/VirtualProduction/Takes/Source/TakeRecorder/Public/ITakeRecorderModule.h
--- a/Engine/Plugins/VirtualProduction/Takes/Source/TakeRecorder/Public/ITakeRecorderModule.h
+++ b/Engine/Plugins/VirtualProduction/Takes/Source/TakeRecorder/Public/ITakeRecorderModule.h
@@ -103,6 +103,16 @@ public:
 	/** The Takes Browser Content Browser Instance Name */
 	static FName TakesBrowserInstanceName;
 
+	static inline ITakeRecorderModule& Get() 
+	{ 
+		return FModuleManager::LoadModuleChecked<ITakeRecorderModule>("TakeRecorder"); 
+	}
+	
+	static inline bool IsAvailable() 
+	{ 
+		return FModuleManager::Get().IsModuleLoaded( "TakeRecorder" ); 
+	}
+
 	/** Get the current pending take. */
 	virtual UTakePreset* GetPendingTake() const = 0;
 
diff --git a/Engine/Plugins/VirtualProduction/Takes/Source/TakeRecorder/Public/Recorder/TakeRecorder.h b/Engine/Plugins/VirtualProduction/Takes/Source/TakeRecorder/Public/Recorder/TakeRecorder.h
--- a/Engine/Plugins/VirtualProduction/Takes/Source/TakeRecorder/Public/Recorder/TakeRecorder.h
+++ b/Engine/Plugins/VirtualProduction/Takes/Source/TakeRecorder/Public/Recorder/TakeRecorder.h
@@ -43,6 +43,7 @@ DECLARE_MULTICAST_DELEGATE_OneParam(FOnTakeRecordingStopped, UTakeRecorder*);
 DECLARE_MULTICAST_DELEGATE_OneParam(FOnTakeRecordingFinished, UTakeRecorder*);
 DECLARE_MULTICAST_DELEGATE_OneParam(FOnTakeRecordingCancelled, UTakeRecorder*);
 DECLARE_MULTICAST_DELEGATE_TwoParams(FOnStartPlayFrameModified, UTakeRecorder*, const FFrameNumber& StartFrame);
+DECLARE_MULTICAST_DELEGATE_TwoParams(FOnTickRecording, UTakeRecorder*, const FQualifiedFrameTime&);
 
 DECLARE_DELEGATE_RetVal_OneParam(FTakeRecorderParameters, FTakeRecorderParameterDelegate,const FTakeRecorderParameters&);
 
@@ -177,6 +178,11 @@ public:
 	 */
 	FOnStartPlayFrameModified& OnStartPlayFrameModified();
 
+	/**
+	 * Retrieve a multi-cast delegate that is triggered when this recording ticks
+	 */
+	FOnTickRecording& OnTickRecording();
+
 private:
 
 	/**
@@ -294,6 +300,9 @@ private:
 	/** Triggered when the movie scene is adjusted. */
 	FOnStartPlayFrameModified OnFrameModifiedEvent;
 
+	/** Triggered when this recorder ticks */
+	FOnTickRecording OnTickRecordingEvent;
+
 	/** Sequencer ptr that controls playback of the destination asset during the recording */
 	TWeakPtr<ISequencer> WeakSequencer;
 
