ShotManagement: Add subsequence list to the cine assembly detail customization.This lets users rename the subsequences defined in the schema, and also selectively choose which ones should actually get created along with the top-level assembly instance.

#rnx

diff --git a/Engine/Plugins/VirtualProduction/Shots/Source/ShotsEditor/Private/CineAssembly.cpp b/Engine/Plugins/VirtualProduction/Shots/Source/ShotsEditor/Private/CineAssembly.cpp
--- a/Engine/Plugins/VirtualProduction/Shots/Source/ShotsEditor/Private/CineAssembly.cpp
+++ b/Engine/Plugins/VirtualProduction/Shots/Source/ShotsEditor/Private/CineAssembly.cpp
@@ -73,6 +73,18 @@ void UCineAssembly::ChangeSchema(UCineAssemblySchema* InSchema)
 			}
 		}
 	}
+
+	// Reset the list of SubAssembly names to create from the Schema
+	SubAssemblyNames.Reset();
+	if (BaseSchema)
+	{
+		for (const FString& SubAssemblyName : BaseSchema->SubsequencesToCreate)
+		{
+			FTemplateString NewTemplateName;
+			NewTemplateName.Template = SubAssemblyName;
+			SubAssemblyNames.Add(NewTemplateName);
+		}
+	}
 }
 
 void UCineAssembly::CreateSubAssemblies()
@@ -108,16 +120,34 @@ void UCineAssembly::CreateSubAssemblies()
 	IAssetTools& AssetTools = FModuleManager::GetModuleChecked<FAssetToolsModule>("AssetTools").Get();
 
 	// Create a new CineAssembly for each subsequence, set its playback range to match the parent sequence, and add it to the subsequence track
-	for (const FString& Subsequence : BaseSchema->SubsequencesToCreate)
+	for (FTemplateString& SubAssemblyName : SubAssemblyNames)
 	{
+		// Resolve any tokens found in the name of the SubAssembly before attempting to create it 
+		UCineAssemblyNamingTokensContext* NamingTokenContext = NewObject<UCineAssemblyNamingTokensContext>();
+		NamingTokenContext->Assembly = this;
+
+		FNamingTokenFilterArgs FilterArgs;
+		FilterArgs.AdditionalNamespacesToInclude.Add(UCineAssemblyNamingTokens::TokenNamespace);
+
+		UNamingTokensEngineSubsystem* NamingTokensSubsystem = GEngine->GetEngineSubsystem<UNamingTokensEngineSubsystem>();
+		FNamingTokenResultData Result = NamingTokensSubsystem->EvaluateTokenString(SubAssemblyName.Template, FilterArgs, { NamingTokenContext });
+		SubAssemblyName.Resolved = Result.EvaluatedText;
+
+		const FString SubAssemblyFilename = FPaths::GetBaseFilename(SubAssemblyName.Resolved.ToString());
+
+		if (SubAssemblyFilename.IsEmpty())
+		{
+			continue;
+		}
+
 		// Add a subsequence track to the assembly's sequence
 		UMovieSceneSubTrack* SubTrack = Cast<UMovieSceneSubTrack>(MovieScene->AddTrack(UMovieSceneSubTrack::StaticClass()));
-		SubTrack->SetDisplayName(FText::FromString(FPaths::GetPathLeaf(Subsequence)));
+		SubTrack->SetDisplayName(FText::FromString(SubAssemblyFilename));
 
 		// Before creating each subassembly, sanity check that each one will actually have a unique name (in case there are duplicates in the schema description)
 		FString UniquePackageName;
 		FString UniqueAssetName;
-		AssetTools.CreateUniqueAssetName(PackagePath / Subsequence, TEXT(""), UniquePackageName, UniqueAssetName);
+		AssetTools.CreateUniqueAssetName(PackagePath / SubAssemblyName.Resolved.ToString(), TEXT(""), UniquePackageName, UniqueAssetName);
 
 		const FString SubAssemblyPath = FPaths::GetPath(UniquePackageName);
 		UCineAssembly* SubAssembly = Cast<UCineAssembly>(AssetTools.CreateAsset(UniqueAssetName, SubAssemblyPath, UCineAssembly::StaticClass(), NewObject<UCineAssemblyFactory>()));
diff --git a/Engine/Plugins/VirtualProduction/Shots/Source/ShotsEditor/Private/CineAssembly.h b/Engine/Plugins/VirtualProduction/Shots/Source/ShotsEditor/Private/CineAssembly.h
--- a/Engine/Plugins/VirtualProduction/Shots/Source/ShotsEditor/Private/CineAssembly.h
+++ b/Engine/Plugins/VirtualProduction/Shots/Source/ShotsEditor/Private/CineAssembly.h
@@ -96,6 +96,10 @@ public:
 	UPROPERTY()
 	FString AssemblyNote;
 
+	/** Array of template names (possibly containing tokens), based on the schema, used to create the SubAssemblies */
+	UPROPERTY()
+	TArray<FTemplateString> SubAssemblyNames;
+
 	/** Array of Subsequence Sections created based on the schema */
 	UPROPERTY()
 	TArray<TObjectPtr<UMovieSceneSubSection>> SubAssemblies;
diff --git a/Engine/Plugins/VirtualProduction/Shots/Source/ShotsEditor/Private/CineAssemblyCustomization.cpp b/Engine/Plugins/VirtualProduction/Shots/Source/ShotsEditor/Private/CineAssemblyCustomization.cpp
--- a/Engine/Plugins/VirtualProduction/Shots/Source/ShotsEditor/Private/CineAssemblyCustomization.cpp
+++ b/Engine/Plugins/VirtualProduction/Shots/Source/ShotsEditor/Private/CineAssemblyCustomization.cpp
@@ -2,10 +2,16 @@
 
 #include "CineAssemblyCustomization.h"
 
+#include "Algo/Contains.h"
+#include "Algo/IndexOf.h"
+#include "Algo/Find.h"
 #include "CineAssembly.h"
+#include "CineAssemblyNamingTokens.h"
 #include "DetailLayoutBuilder.h"
 #include "DetailWidgetRow.h"
+#include "NamingTokensEngineSubsystem.h"
 #include "PropertyCustomizationHelpers.h"
+#include "STemplateStringEditableTextBox.h"
 #include "Widgets/Input/SCheckBox.h"
 #include "Widgets/Input/SMultiLineEditableTextBox.h"
 #include "Widgets/Input/SNumericEntryBox.h"
@@ -38,8 +44,12 @@ void FCineAssemblyCustomization::CustomizeDetails(IDetailLayoutBuilder& DetailBu
 		DetailBuilder.EditCategory("TargetLevel", LOCTEXT("TargetLevelCategoryName", "Target Level"));
 
 		// Add a new category for Schema Metadata properties
-		const FText CategoryName = FText::Format(LOCTEXT("SchemaMetadataCategoryName", "{0} Metadata"), FText::FromString(BaseSchema->SchemaName));
-		IDetailCategoryBuilder& MetadataCategory = DetailBuilder.EditCategory("SchemaMetadata", CategoryName);
+		const FText MetadataCategoryName = FText::Format(LOCTEXT("SchemaMetadataCategoryName", "{0} Metadata"), FText::FromString(BaseSchema->SchemaName));
+		IDetailCategoryBuilder& MetadataCategory = DetailBuilder.EditCategory("SchemaMetadata", MetadataCategoryName);
+
+		// Add a new category for Schema Subsequences
+		const FText SubsequencesCategoryName = LOCTEXT("SchemaSubseuqnecesCategoryName", "Subsequences");
+		IDetailCategoryBuilder& SubsequenceCategory = DetailBuilder.EditCategory("SchemaSubsequences", SubsequencesCategoryName);
 
 		// Add a property row for each metadata struct in the customized Assembly's base schema
 		for (const FAssemblyMetadataDesc& MetadataDesc : BaseSchema->AssemblyMetadata)
@@ -178,7 +188,127 @@ void FCineAssemblyCustomization::CustomizeDetails(IDetailLayoutBuilder& DetailBu
 					ValueWidget.ToSharedRef()
 				];
 		}
+
+		SubAssemblyNames.Reset();
+		for (const FString& SubsequenceName : BaseSchema->SubsequencesToCreate)
+		{
+			FTemplateString SubsequenceTemplateString;
+			SubsequenceTemplateString.Template = SubsequenceName;
+			SubAssemblyNames.Add(SubsequenceTemplateString);
+		}
+
+		for (int32 Index = 0; Index < SubAssemblyNames.Num(); ++Index)
+		{
+			SubsequenceCategory.AddCustomRow(FText::GetEmpty())
+				.NameContent()
+				.HAlign(HAlign_Right)
+				[
+					SNew(SCheckBox)
+						.IsChecked_Raw(this, &FCineAssemblyCustomization::IsSubAssemblyChecked, CustomizedCineAssembly, Index)
+						.OnCheckStateChanged_Raw(this, &FCineAssemblyCustomization::SubAssemblyCheckStateChanged, CustomizedCineAssembly, Index)
+				]
+				.ValueContent()
+				.HAlign(HAlign_Fill)
+				[
+					SNew(STemplateStringEditableTextBox)
+						.Text_Raw(this, &FCineAssemblyCustomization::GetTemplateText, Index)
+						.ResolvedText_Lambda([this, Index]() { return GetResolvedText(Index); })
+						.OnTextCommitted_Raw(this, &FCineAssemblyCustomization::OnTemplateTextCommitted, CustomizedCineAssembly, Index)
+				];
+		}
+	}
+}
+
+ECheckBoxState FCineAssemblyCustomization::IsSubAssemblyChecked(UCineAssembly* Assembly, int32 Index) const
+{
+	if (SubAssemblyNames.IsValidIndex(Index))
+	{
+		const FTemplateString& SubAssemblyName = SubAssemblyNames[Index];
+		const bool bResult = Algo::ContainsBy(Assembly->SubAssemblyNames, SubAssemblyName.Template, &FTemplateString::Template);
+		return bResult ? ECheckBoxState::Checked : ECheckBoxState::Unchecked;
+	}
+	return ECheckBoxState::Undetermined;
+}
+
+void FCineAssemblyCustomization::SubAssemblyCheckStateChanged(ECheckBoxState CheckBoxState, UCineAssembly* Assembly, int32 Index)
+{
+	if (SubAssemblyNames.IsValidIndex(Index))
+	{
+		FTemplateString& SubAssemblyName = SubAssemblyNames[Index];
+
+		if (CheckBoxState == ECheckBoxState::Checked)
+		{
+			Assembly->SubAssemblyNames.Add(SubAssemblyName);
+		}
+		else if (CheckBoxState == ECheckBoxState::Unchecked)
+		{
+			const int32 IndexToRemove = Algo::IndexOfBy(Assembly->SubAssemblyNames, SubAssemblyName.Template, &FTemplateString::Template);
+			Assembly->SubAssemblyNames.RemoveAt(IndexToRemove);
+		}
+	}
+}
+
+FText FCineAssemblyCustomization::GetTemplateText(int32 Index) const
+{
+	if (SubAssemblyNames.IsValidIndex(Index))
+	{
+		const FTemplateString& SubAssemblyName = SubAssemblyNames[Index];
+		return FText::FromString(FPaths::GetBaseFilename(SubAssemblyName.Template));
+	}
+	return FText::GetEmpty();
+}
+
+FText FCineAssemblyCustomization::GetResolvedText(int32 Index)
+{
+	EvaluateTokenStrings();
+
+	if (SubAssemblyNames.IsValidIndex(Index))
+	{
+		const FTemplateString& SubAssemblyName = SubAssemblyNames[Index];
+		return FText::FromString(FPaths::GetBaseFilename(SubAssemblyName.Resolved.ToString()));
+	}
+	return FText::GetEmpty();
+}
+
+void FCineAssemblyCustomization::OnTemplateTextCommitted(const FText& InText, ETextCommit::Type InCommitType, UCineAssembly* Assembly, int32 Index)
+{
+	if (SubAssemblyNames.IsValidIndex(Index))
+	{
+		FTemplateString& SubAssemblyName = SubAssemblyNames[Index];
+
+		FTemplateString* TemplateString = Algo::FindBy(Assembly->SubAssemblyNames, SubAssemblyName.Template, &FTemplateString::Template);
+
+		const FString Path = FPaths::GetPath(SubAssemblyName.Template);
+		SubAssemblyName.Template = Path / InText.ToString();
+
+		if (TemplateString)
+		{
+			TemplateString->Template = SubAssemblyName.Template;
+		}
+
+		EvaluateTokenString(SubAssemblyName);
 	}
 }
 
+void FCineAssemblyCustomization::EvaluateTokenStrings()
+{
+	FDateTime CurrentTime = FDateTime::Now();
+	if ((CurrentTime - LastTokenUpdateTime).GetSeconds() >= 1.0f)
+	{
+		for (FTemplateString& SubAssemblyName : SubAssemblyNames)
+		{
+			EvaluateTokenString(SubAssemblyName);
+		}
+
+		LastTokenUpdateTime = CurrentTime;
+	}
+}
+
+void FCineAssemblyCustomization::EvaluateTokenString(FTemplateString& TokenString)
+{
+	UNamingTokensEngineSubsystem* NamingTokensSubsystem = GEngine->GetEngineSubsystem<UNamingTokensEngineSubsystem>();
+	FNamingTokenResultData Result = NamingTokensSubsystem->EvaluateTokenString(TokenString.Template, FNamingTokenFilterArgs({ UCineAssemblyNamingTokens::TokenNamespace }));
+	TokenString.Resolved = Result.EvaluatedText;
+}
+
 #undef LOCTEXT_NAMESPACE
diff --git a/Engine/Plugins/VirtualProduction/Shots/Source/ShotsEditor/Private/CineAssemblyCustomization.h b/Engine/Plugins/VirtualProduction/Shots/Source/ShotsEditor/Private/CineAssemblyCustomization.h
--- a/Engine/Plugins/VirtualProduction/Shots/Source/ShotsEditor/Private/CineAssemblyCustomization.h
+++ b/Engine/Plugins/VirtualProduction/Shots/Source/ShotsEditor/Private/CineAssemblyCustomization.h
@@ -4,6 +4,11 @@
 
 #include "IDetailCustomization.h"
 
+#include "Styling/SlateTypes.h"
+#include "UObject/TemplateString.h"
+
+class UCineAssembly;
+
 /**
  * Detail customization for UCineAssembly
  */
@@ -16,4 +21,39 @@ public:
 	//~ Begin IDetailCustomization interface */
 	virtual void CustomizeDetails(IDetailLayoutBuilder& DetailBuilder) override;
 	//~ End IDetailCustomization interface */
+
+private:
+	/**
+	 * Evaluates all of the customization's template strings with the naming tokens subsystem.
+	 * This function is throttled to only run at a set frequency, to avoid the potential to constantly query the naming tokens subsystem.
+	 */
+	void EvaluateTokenStrings();
+
+	/**
+	 * Evaluates the input template string with the naming tokens subsystem.
+	 * This function is not throttled to allow for immediate updates.
+	 */
+	void EvaluateTokenString(FTemplateString& TokenString);
+
+	/** Checks if the customized CineAssembly contains a SubAssembly name in its list of assets to create */
+	ECheckBoxState IsSubAssemblyChecked(UCineAssembly* Assembly, int32 Index) const;
+
+	/** Adds/Removes a SubAssembly name from the custmoized CineAssembly's list of SubAssemblies */
+	void SubAssemblyCheckStateChanged(ECheckBoxState CheckBoxState, UCineAssembly* Assembly, int32 Index);
+
+	/** Returns the template text for a SubAssembly template name */
+	FText GetTemplateText(int32 Index) const;
+
+	/** Evaluates the token strings, then returns the resolved text for a SubAssembly template name */
+	FText GetResolvedText(int32 Index);
+
+	/** Modifies the template text of the SubAssembly template name and re-evaluates the token string to update the resolved text */
+	void OnTemplateTextCommitted(const FText& InText, ETextCommit::Type InCommitType, UCineAssembly* Assembly, int32 Index);
+
+private:
+	/** Array of template names from the customized CineAssembly's schema */
+	TArray<FTemplateString> SubAssemblyNames;
+
+	/** The last time the naming tokens were updated */
+	FDateTime LastTokenUpdateTime;
 };
diff --git a/Engine/Plugins/VirtualProduction/Shots/Source/ShotsEditor/Private/UI/CineAssembly/SCineAssemblyConfigWindow.cpp b/Engine/Plugins/VirtualProduction/Shots/Source/ShotsEditor/Private/UI/CineAssembly/SCineAssemblyConfigWindow.cpp
--- a/Engine/Plugins/VirtualProduction/Shots/Source/ShotsEditor/Private/UI/CineAssembly/SCineAssemblyConfigWindow.cpp
+++ b/Engine/Plugins/VirtualProduction/Shots/Source/ShotsEditor/Private/UI/CineAssembly/SCineAssemblyConfigWindow.cpp
@@ -291,9 +291,9 @@ void SCineAssemblyConfigWindow::EvaluateTokenString(FTemplateString& StringToEva
 	if ((CurrentTime - LastTokenUpdateTime).GetSeconds() >= 1.0f)
 	{
 		UNamingTokensEngineSubsystem* NamingTokensSubsystem = GEngine->GetEngineSubsystem<UNamingTokensEngineSubsystem>();
-		FNamingTokenResultData Result = NamingTokensSubsystem->EvaluateTokenString(CineAssemblyToConfigure->AssemblyName.Template, FNamingTokenFilterArgs({ UCineAssemblyNamingTokens::TokenNamespace }));
+		FNamingTokenResultData Result = NamingTokensSubsystem->EvaluateTokenString(StringToEvaluate.Template, FNamingTokenFilterArgs({ UCineAssemblyNamingTokens::TokenNamespace }));
 
-		CineAssemblyToConfigure->AssemblyName.Resolved = Result.EvaluatedText;
+		StringToEvaluate.Resolved = Result.EvaluatedText;
 		LastTokenUpdateTime = CurrentTime;
 	}
 }
