ShotManagement: Add a hierarchy tree view to the Cine Assembly creation window that displays a read-only view of the folders and subassemblies that will get created based on the selected schema (all with their tokenized names resolved).

#rnx

diff --git a/Engine/Plugins/VirtualProduction/Shots/Source/ShotsEditor/Private/CineAssembly.cpp b/Engine/Plugins/VirtualProduction/Shots/Source/ShotsEditor/Private/CineAssembly.cpp
--- a/Engine/Plugins/VirtualProduction/Shots/Source/ShotsEditor/Private/CineAssembly.cpp
+++ b/Engine/Plugins/VirtualProduction/Shots/Source/ShotsEditor/Private/CineAssembly.cpp
@@ -85,6 +85,18 @@ void UCineAssembly::ChangeSchema(UCineAssemblySchema* InSchema)
 			SubAssemblyNames.Add(NewTemplateName);
 		}
 	}
+
+	// Reset the list of folders names to create from the Schema
+	DefaultFolderNames.Reset();
+	if (BaseSchema)
+	{
+		for (const FString& FolderName : BaseSchema->FoldersToCreate)
+		{
+			FTemplateString NewTemplateName;
+			NewTemplateName.Template = FolderName;
+			DefaultFolderNames.Add(NewTemplateName);
+		}
+	}
 }
 
 void UCineAssembly::CreateSubAssemblies()
@@ -94,14 +106,31 @@ void UCineAssembly::CreateSubAssemblies()
 		return;
 	}
 
+	UNamingTokensEngineSubsystem* NamingTokensSubsystem = GEngine->GetEngineSubsystem<UNamingTokensEngineSubsystem>();
+
+	UCineAssemblyNamingTokensContext* NamingTokenContext = NewObject<UCineAssemblyNamingTokensContext>();
+	NamingTokenContext->Assembly = this;
+
+	FNamingTokenFilterArgs FilterArgs;
+	FilterArgs.AdditionalNamespacesToInclude.Add(UCineAssemblyNamingTokens::TokenNamespace);
+
 	// Get the path where the top-level assembly will be created so we can create other assets relative to it
 	const FAssetData AssemblyAssetData = FAssetData(this);
 	const FString PackagePath = AssemblyAssetData.PackagePath.ToString();
 
 	// Create the default folders for this assembly, based on the schema
-	for (const FString& FolderPath : BaseSchema->FoldersToCreate)
+	for (FTemplateString& FolderPath : DefaultFolderNames)
 	{
-		const FString PathToCreate = PackagePath / FolderPath;
+		// Resolve any tokens found in the name of the SubAssembly before attempting to create it 
+		FNamingTokenResultData Result = NamingTokensSubsystem->EvaluateTokenString(FolderPath.Template, FilterArgs, { NamingTokenContext });
+		FolderPath.Resolved = Result.EvaluatedText;
+
+		if (FolderPath.Resolved.IsEmpty())
+		{
+			continue;
+		}
+
+		const FString PathToCreate = PackagePath / FolderPath.Resolved.ToString();
 		const FString RelativeFilePath = FPackageName::LongPackageNameToFilename(PathToCreate);
 		const FString AbsoluteFilePath = FPaths::ConvertRelativePathToFull(RelativeFilePath);
 
@@ -123,13 +152,6 @@ void UCineAssembly::CreateSubAssemblies()
 	for (FTemplateString& SubAssemblyName : SubAssemblyNames)
 	{
 		// Resolve any tokens found in the name of the SubAssembly before attempting to create it 
-		UCineAssemblyNamingTokensContext* NamingTokenContext = NewObject<UCineAssemblyNamingTokensContext>();
-		NamingTokenContext->Assembly = this;
-
-		FNamingTokenFilterArgs FilterArgs;
-		FilterArgs.AdditionalNamespacesToInclude.Add(UCineAssemblyNamingTokens::TokenNamespace);
-
-		UNamingTokensEngineSubsystem* NamingTokensSubsystem = GEngine->GetEngineSubsystem<UNamingTokensEngineSubsystem>();
 		FNamingTokenResultData Result = NamingTokensSubsystem->EvaluateTokenString(SubAssemblyName.Template, FilterArgs, { NamingTokenContext });
 		SubAssemblyName.Resolved = Result.EvaluatedText;
 
diff --git a/Engine/Plugins/VirtualProduction/Shots/Source/ShotsEditor/Private/CineAssembly.h b/Engine/Plugins/VirtualProduction/Shots/Source/ShotsEditor/Private/CineAssembly.h
--- a/Engine/Plugins/VirtualProduction/Shots/Source/ShotsEditor/Private/CineAssembly.h
+++ b/Engine/Plugins/VirtualProduction/Shots/Source/ShotsEditor/Private/CineAssembly.h
@@ -100,6 +100,10 @@ public:
 	UPROPERTY()
 	TArray<FTemplateString> SubAssemblyNames;
 
+	/** Array of template names (possibly containing tokens), based on the schema, used to create folders for this assembly */
+	UPROPERTY()
+	TArray<FTemplateString> DefaultFolderNames;
+
 	/** Array of Subsequence Sections created based on the schema */
 	UPROPERTY()
 	TArray<TObjectPtr<UMovieSceneSubSection>> SubAssemblies;
diff --git a/Engine/Plugins/VirtualProduction/Shots/Source/ShotsEditor/Private/UI/CineAssembly/SCineAssemblyConfigWindow.cpp b/Engine/Plugins/VirtualProduction/Shots/Source/ShotsEditor/Private/UI/CineAssembly/SCineAssemblyConfigWindow.cpp
--- a/Engine/Plugins/VirtualProduction/Shots/Source/ShotsEditor/Private/UI/CineAssembly/SCineAssemblyConfigWindow.cpp
+++ b/Engine/Plugins/VirtualProduction/Shots/Source/ShotsEditor/Private/UI/CineAssembly/SCineAssemblyConfigWindow.cpp
@@ -178,23 +178,13 @@ void SCineAssemblyConfigWindow::OnSchemaSelected(const FAssetData& AssetData)
 
 	// The details view needs to be redrawn to show the new metadata fields from the selected schema
 	DetailsView->ForceRefresh();
+
+	// Recreate the hierarchy tree items based on the selected schema
+	PopulateHierarchyTree();
 }
 
 TSharedRef<SWidget> SCineAssemblyConfigWindow::MakeInfoPanel()
 {
-	// TODO: Remove this when all of the tab content has been implemented
-	auto BuildStubWidget = [](const FString& StubText) -> TSharedRef<SWidget>
-		{
-			return SNew(SHorizontalBox)
-				+ SHorizontalBox::Slot()
-				.FillWidth(1.0f)
-				.VAlign(VAlign_Center)
-				.HAlign(HAlign_Center)
-				[
-					SNew(STextBlock).Text(FText::FromString(StubText))
-				];
-		};
-
 	TabSwitcher = SNew(SWidgetSwitcher)
 		+ SWidgetSwitcher::Slot()
 			[ MakeOverviewWidget() ]
@@ -203,7 +193,7 @@ TSharedRef<SWidget> SCineAssemblyConfigWindow::MakeInfoPanel()
 			[ MakeDetailsWidget() ]
 
 		+ SWidgetSwitcher::Slot()
-			[ BuildStubWidget("HIERARCHY TAB CONTENT") ]
+			[ MakeHierarchyWidget() ]
 
 		+ SWidgetSwitcher::Slot()
 			[ MakeNotesWidget() ];
@@ -382,6 +372,216 @@ TSharedRef<SWidget> SCineAssemblyConfigWindow::MakeDetailsWidget()
 	return DetailsView.ToSharedRef();
 }
 
+TSharedRef<SWidget> SCineAssemblyConfigWindow::MakeHierarchyWidget()
+{
+	HierarchyTreeView = SNew(STreeView<TSharedPtr<FHierarchyTreeItem>>)
+		.TreeItemsSource(&HierarchyTreeItems)
+		.SelectionMode(ESelectionMode::None)
+		.OnGenerateRow(this, &SCineAssemblyConfigWindow::OnGenerateTreeRow)
+		.OnGetChildren(this, &SCineAssemblyConfigWindow::OnGetChildren);
+
+	// Create the hierarchy tree root
+	RootItem = MakeShared<FHierarchyTreeItem>();
+	RootItem->Type = FHierarchyTreeItem::EItemType::Folder;
+
+	FTemplateString RootTemplate;
+	RootTemplate.Template = TEXT("");
+	RootItem->Path = RootTemplate;
+
+	PopulateHierarchyTree();
+
+	// Register a Slate timer that runs at a set frequency to evaluate all of the tokens in the tree view.
+	// This will automatically be unregistered when this window is destroyed.
+	constexpr float TimerFrequency = 1.0f;
+	RegisterActiveTimer(TimerFrequency, FWidgetActiveTimerDelegate::CreateLambda([this](double InCurrentTime, float InDeltaTime) -> EActiveTimerReturnType
+		{
+			EvaluateHierarchyTokensRecursive(RootItem);
+			HierarchyTreeView->RequestTreeRefresh();
+
+			return EActiveTimerReturnType::Continue;
+		}));
+
+	return SNew(SVerticalBox)
+
+		+ SVerticalBox::Slot()
+		.Padding(16.0f)
+		.AutoHeight()
+		[
+			SNew(STextBlock)
+				.Text(LOCTEXT("HierarchyInstructions", "The following hierarchy will be created as defined by the shot type."))
+				.AutoWrapText(true)
+		]
+
+		+ SVerticalBox::Slot()
+		[
+			SNew(SBorder)
+				.BorderImage(FShotsEditorStyle::Get().GetBrush("ProductionWizard.RecessedBackground"))
+				.Padding(8.0f)
+				[
+					HierarchyTreeView.ToSharedRef()
+				]
+		];
+}
+
+void SCineAssemblyConfigWindow::PopulateHierarchyTree()
+{
+	HierarchyTreeItems.Empty();
+	HierarchyTreeItems.Add(RootItem);
+
+	auto AddItemsToTree = [this](TArray<FTemplateString>& ItemList, FHierarchyTreeItem::EItemType ItemType)
+		{
+			// Sort the list so that paths are added to the tree in the proper order
+			ItemList.Sort([](const FTemplateString& A, const FTemplateString& B) { return A.Template < B.Template; });
+
+			for (const FTemplateString& ItemName : ItemList)
+			{
+				const FString ParentPath = FPaths::GetPath(ItemName.Template);
+
+				// Walk the tree until we find an item whose path matches the parent path. The new tree item will be created as one of its children
+				TSharedPtr<FHierarchyTreeItem> ParentItem = FindItemAtPathRecursive(RootItem, ParentPath);
+				if (ParentItem)
+				{
+					TSharedPtr<FHierarchyTreeItem> NewItem = MakeShared<FHierarchyTreeItem>();
+					NewItem->Type = ItemType;
+					NewItem->Path = ItemName;
+
+					if (ItemType == FHierarchyTreeItem::EItemType::Folder)
+					{
+						ParentItem->ChildFolders.Add(NewItem);
+					}
+					else
+					{
+						ParentItem->ChildAssets.Add(NewItem);
+					}
+				}
+			}
+		};
+
+	if (CineAssemblyToConfigure->BaseSchema)
+	{
+		TArray<FTemplateString> FolderTemplates;
+		FolderTemplates.Reserve(CineAssemblyToConfigure->BaseSchema->FoldersToCreate.Num());
+		Algo::Transform(CineAssemblyToConfigure->BaseSchema->FoldersToCreate, FolderTemplates, [](const FString& TemplateString)
+			{ 
+				FTemplateString FolderTemplate;
+				FolderTemplate.Template = TemplateString;
+				return FolderTemplate;
+			});
+
+		AddItemsToTree(FolderTemplates, FHierarchyTreeItem::EItemType::Folder);
+		AddItemsToTree(CineAssemblyToConfigure->SubAssemblyNames, FHierarchyTreeItem::EItemType::Asset);
+
+		EvaluateHierarchyTokensRecursive(RootItem);
+	}
+
+	HierarchyTreeView->RequestTreeRefresh();
+	ExpandTreeRecursive(RootItem);
+}
+
+void SCineAssemblyConfigWindow::EvaluateHierarchyTokensRecursive(TSharedPtr<FHierarchyTreeItem> TreeItem)
+{
+	UNamingTokensEngineSubsystem* NamingTokensSubsystem = GEngine->GetEngineSubsystem<UNamingTokensEngineSubsystem>();
+
+	// Evalaute the token template string for this tree item
+	FNamingTokenResultData Result = NamingTokensSubsystem->EvaluateTokenString(TreeItem->Path.Template, FNamingTokenFilterArgs({ UCineAssemblyNamingTokens::TokenNamespace }));
+	TreeItem->Path.Resolved = Result.EvaluatedText;
+
+	// Evaluate the tokens for all of the child assets, then resort them alphabetically based on the resolved paths
+	for (const TSharedPtr<FHierarchyTreeItem>& Asset : TreeItem->ChildAssets)
+	{
+		EvaluateHierarchyTokensRecursive(Asset);
+	}
+
+	TreeItem->ChildAssets.Sort([](const TSharedPtr<FHierarchyTreeItem>& A, const TSharedPtr<FHierarchyTreeItem>& B) { return A->Path.Resolved.ToString() < B->Path.Resolved.ToString(); });
+
+	// Evaluate the tokens for all of the child folders, then resort them alphabetically based on the resolved paths
+	for (const TSharedPtr<FHierarchyTreeItem>& Child : TreeItem->ChildFolders)
+	{
+		EvaluateHierarchyTokensRecursive(Child);
+	}
+
+	TreeItem->ChildFolders.Sort([](const TSharedPtr<FHierarchyTreeItem>& A, const TSharedPtr<FHierarchyTreeItem>& B) { return A->Path.Resolved.ToString() < B->Path.Resolved.ToString(); });
+}
+
+void SCineAssemblyConfigWindow::ExpandTreeRecursive(TSharedPtr<FHierarchyTreeItem> TreeItem) const
+{
+	if (HierarchyTreeView)
+	{
+		HierarchyTreeView->SetItemExpansion(TreeItem, true);
+	}
+
+	for (const TSharedPtr<FHierarchyTreeItem>& ChildItem : TreeItem->ChildFolders)
+	{
+		ExpandTreeRecursive(ChildItem);
+	}
+}
+
+TSharedPtr<FHierarchyTreeItem> SCineAssemblyConfigWindow::FindItemAtPathRecursive(TSharedPtr<FHierarchyTreeItem> TreeItem, const FString& Path) const
+{
+	if (TreeItem->Path.Template.Equals(Path))
+	{
+		return TreeItem;
+	}
+
+	for (const TSharedPtr<FHierarchyTreeItem>& Child : TreeItem->ChildFolders)
+	{
+		if (const TSharedPtr<FHierarchyTreeItem>& ItemAtPath = FindItemAtPathRecursive(Child, Path))
+		{
+			return ItemAtPath;
+		}
+	}
+
+	return nullptr;
+}
+
+TSharedRef<ITableRow> SCineAssemblyConfigWindow::OnGenerateTreeRow(TSharedPtr<FHierarchyTreeItem> TreeItem, const TSharedRef<STableViewBase>& OwnerTable)
+{
+	TSharedPtr<SImage> Icon = SNew(SImage);
+
+	if (TreeItem->Type == FHierarchyTreeItem::EItemType::Folder)
+	{
+		Icon->SetImage(FShotsEditorStyle::Get().GetBrush("Icons.Folder"));
+		Icon->SetColorAndOpacity(FAppStyle::Get().GetSlateColor("ContentBrowser.DefaultFolderColor"));
+	}
+	else
+	{
+		Icon->SetImage(FShotsEditorStyle::Get().GetBrush("Icons.Sequencer"));
+		Icon->SetColorAndOpacity(FLinearColor::White);
+	}
+
+	return SNew(STableRow<TSharedPtr<FHierarchyTreeItem>>, OwnerTable)
+		.Padding(FMargin(8.0f, 2.0f, 8.0f, 0.0f))
+		.Content()
+		[
+			SNew(SHorizontalBox)
+
+			+ SHorizontalBox::Slot()
+				.AutoWidth()
+				.VAlign(VAlign_Center)
+				.Padding(0.0f, 0.0f, 2.0f, 0.0f)
+				[
+					Icon.ToSharedRef()
+				]
+
+			+ SHorizontalBox::Slot()
+				.VAlign(VAlign_Center)
+				[
+					SNew(STextBlock)
+						.Text_Lambda([this, TreeItem]() -> FText
+							{
+								return (TreeItem == RootItem) ? LOCTEXT("RootPathName", "Root Folder") : FText::FromString(FPaths::GetPathLeaf(TreeItem->Path.Resolved.ToString()));
+							})
+				]
+		];
+}
+
+void SCineAssemblyConfigWindow::OnGetChildren(TSharedPtr<FHierarchyTreeItem> TreeItem, TArray<TSharedPtr<FHierarchyTreeItem>>& OutNodes)
+{
+	// Display all of the child assets first, followed by all of the child folders
+	OutNodes.Append(TreeItem->ChildAssets);
+	OutNodes.Append(TreeItem->ChildFolders);
+}
+
 TSharedRef<SWidget> SCineAssemblyConfigWindow::MakeNotesWidget()
 {
 	return SNew(SVerticalBox)
diff --git a/Engine/Plugins/VirtualProduction/Shots/Source/ShotsEditor/Private/UI/CineAssembly/SCineAssemblyConfigWindow.h b/Engine/Plugins/VirtualProduction/Shots/Source/ShotsEditor/Private/UI/CineAssembly/SCineAssemblyConfigWindow.h
--- a/Engine/Plugins/VirtualProduction/Shots/Source/ShotsEditor/Private/UI/CineAssembly/SCineAssemblyConfigWindow.h
+++ b/Engine/Plugins/VirtualProduction/Shots/Source/ShotsEditor/Private/UI/CineAssembly/SCineAssemblyConfigWindow.h
@@ -9,6 +9,7 @@
 #include "IDetailsView.h"
 #include "Widgets/Layout/SWidgetSwitcher.h"
 #include "Widgets/SWindow.h"
+#include "Widgets/Views/STreeView.h"
 
 /** Mode that configures the UI based on the intended user interactions with the Cine Assembly asset */
 enum class ECineAssemblyConfigMode : uint8
@@ -17,6 +18,29 @@ enum class ECineAssemblyConfigMode : uint8
 	Edit
 };
 
+/** An entry in the hierarchy tree view */
+struct FHierarchyTreeItem
+{
+	/** The types of items that can be represented in this tree view */
+	enum class EItemType : uint8
+	{
+		Asset,
+		Folder
+	};
+
+	/** The type of this tree item */
+	EItemType Type;
+
+	/** The relative path of this tree item (possibly containing tokens) */
+	FTemplateString Path;
+
+	/** The children of this item in the tree that are Asset types */
+	TArray<TSharedPtr<FHierarchyTreeItem>> ChildAssets;
+
+	/** The children of this item in the tree that are Folder types */
+	TArray<TSharedPtr<FHierarchyTreeItem>> ChildFolders;
+};
+
 /**
  * A window for configuring the properties of a UCineAssembly asset
  */
@@ -49,6 +73,9 @@ private:
 	/** Creates the widget to display for the Details tab */
 	TSharedRef<SWidget> MakeDetailsWidget();
 
+	/** Creates the widget to display for the Hierarchy tab */
+	TSharedRef<SWidget> MakeHierarchyWidget();
+
 	/** Creates the widget to display for the Notes tab */
 	TSharedRef<SWidget> MakeNotesWidget();
 
@@ -67,6 +94,24 @@ private:
 	 */
 	void EvaluateTokenString(FTemplateString& StringToEvaluate);
 
+	/** Populate the tree view items from the list of folders and assets specified by the selected schema */
+	void PopulateHierarchyTree();
+
+	/** Generates the row widget for an entry in the tree view */
+	TSharedRef<ITableRow> OnGenerateTreeRow(TSharedPtr<FHierarchyTreeItem> TreeItem, const TSharedRef<STableViewBase>& OwnerTable);
+
+	/** Gets the children of the input tree view item to build additional tree rows */
+	void OnGetChildren(TSharedPtr<FHierarchyTreeItem> TreeItem, TArray<TSharedPtr<FHierarchyTreeItem>>& OutNodes);
+
+	/** Returns the tree item whose path matches the input path */
+	TSharedPtr<FHierarchyTreeItem> FindItemAtPathRecursive(TSharedPtr<FHierarchyTreeItem> TreeItem, const FString& Path) const;
+
+	/** Recursively expands every item in the tree view */
+	void ExpandTreeRecursive(TSharedPtr<FHierarchyTreeItem> TreeItem) const;
+
+	/** Recursively evaluates the tokens in the path of each item in the tree view */
+	void EvaluateHierarchyTokensRecursive(TSharedPtr<FHierarchyTreeItem> TreeItem);
+
 	/** Callback when this window is about to be closed */
 	void OnWindowClosed(const TSharedRef<SWindow>& InWindow);
 
@@ -91,4 +136,13 @@ private:
 
 	/** The last time the naming tokens were updated */
 	FDateTime LastTokenUpdateTime;
+
+	/** Items source for the tree view */
+	TArray<TSharedPtr<FHierarchyTreeItem>> HierarchyTreeItems;
+
+	/** The root item in the tree view */
+	TSharedPtr<FHierarchyTreeItem> RootItem;
+
+	/** A read-only tree view of folders and assets that will be created based on the selected schema */
+	TSharedPtr<STreeView<TSharedPtr<FHierarchyTreeItem>>> HierarchyTreeView;
 };
diff --git a/Engine/Plugins/VirtualProduction/Shots/Source/ShotsEditor/Private/UI/CineAssembly/SCineAssemblySchemaWindow.cpp b/Engine/Plugins/VirtualProduction/Shots/Source/ShotsEditor/Private/UI/CineAssembly/SCineAssemblySchemaWindow.cpp
--- a/Engine/Plugins/VirtualProduction/Shots/Source/ShotsEditor/Private/UI/CineAssembly/SCineAssemblySchemaWindow.cpp
+++ b/Engine/Plugins/VirtualProduction/Shots/Source/ShotsEditor/Private/UI/CineAssembly/SCineAssemblySchemaWindow.cpp
@@ -888,7 +888,13 @@ bool SCineAssemblySchemaWindow::ValidateAssetName(const FText& InText, FText& Ou
 
 	// Ensure that the name does not contain any characters that would be invalid for an asset name
 	// This matches the validation that would happen if the user was renaming an asset in the content browser
-	return FName::IsValidXName(InText.ToString(), INVALID_OBJECTNAME_CHARACTERS INVALID_LONGPACKAGE_CHARACTERS, &OutErrorMessage);
+	FString InvalidCharacters = INVALID_OBJECTNAME_CHARACTERS INVALID_LONGPACKAGE_CHARACTERS;
+
+	// These characters are actually valid, because we want to support naming tokens
+	InvalidCharacters = InvalidCharacters.Replace(TEXT("{}"), TEXT(""));
+	InvalidCharacters = InvalidCharacters.Replace(TEXT(":"), TEXT(""));
+
+	return FName::IsValidXName(InText.ToString(), InvalidCharacters, &OutErrorMessage);
 }
 
 bool SCineAssemblySchemaWindow::ValidateFolderName(const FText& InText, FText& OutErrorMessage, TSharedPtr<FSchemaTreeItem> TreeItem) const
