ShotManagement: Added a segmented control widget to Cine Assembly Schema window. Broke up the details panel into a Details Tab, which shows most of the Schema properties, and a Data Tab, which only displays the metadata properties.

#rnx

diff --git a/Engine/Plugins/VirtualProduction/Shots/Source/ShotsEditor/Private/ShotsEditorStyle.cpp b/Engine/Plugins/VirtualProduction/Shots/Source/ShotsEditor/Private/ShotsEditorStyle.cpp
--- a/Engine/Plugins/VirtualProduction/Shots/Source/ShotsEditor/Private/ShotsEditorStyle.cpp
+++ b/Engine/Plugins/VirtualProduction/Shots/Source/ShotsEditor/Private/ShotsEditorStyle.cpp
@@ -38,6 +38,7 @@ FShotsEditorStyle::FShotsEditorStyle()
 	Set("Icons.Animation", new FSlateVectorImageBrush(RootToContentDir(TEXT("Slate/Animation_16.svg")), Icon16x16));
 	Set("Icons.Notes", new FSlateVectorImageBrush(RootToContentDir(TEXT("Slate/Note.svg")), Icon16x16));
 	Set("Icons.Shot", new FSlateVectorImageBrush(RootToContentDir(TEXT("Slate/Shot_16.svg")), Icon16x16));
+	Set("Icons.DataAsset", new FSlateVectorImageBrush(RootToContentDir(TEXT("Slate/DataAsset.svg")), Icon16x16));
 	Set("Icons.DragHandle", new FSlateVectorImageBrush(RootToContentDir(TEXT("Slate/drag-handle.svg")), Icon16x16));
 
 	Set("Badges.RevisionControlWarning", new FSlateVectorImageBrush(RootToCoreContentDir(TEXT("Starship/SourceControl/Status/RevisionControlBadgeWarning.svg")), Icon16x16, FStyleColors::Warning));
@@ -55,6 +56,7 @@ FShotsEditorStyle::FShotsEditorStyle()
 	Set("ProductionWizard.RecessedBackground", new FSlateRoundedBoxBrush(FStyleColors::Recessed, 0.0f, FStyleColors::Panel, 2.0f));
 
 	Set("Borders.Background", new FSlateRoundedBoxBrush(FStyleColors::Background, 0.0f, FStyleColors::Panel, 2.0f));
+	Set("Borders.PanelNoBorder", new FSlateRoundedBoxBrush(FStyleColors::Panel, 0.0f, FStyleColors::Panel, 0.0f));
 
 	// Buttons
 	const FButtonStyle RecessedButton = FButtonStyle()
diff --git a/Engine/Plugins/VirtualProduction/Shots/Source/ShotsEditor/Private/UI/CineAssembly/SCineAssemblySchemaWindow.cpp b/Engine/Plugins/VirtualProduction/Shots/Source/ShotsEditor/Private/UI/CineAssembly/SCineAssemblySchemaWindow.cpp
--- a/Engine/Plugins/VirtualProduction/Shots/Source/ShotsEditor/Private/UI/CineAssembly/SCineAssemblySchemaWindow.cpp
+++ b/Engine/Plugins/VirtualProduction/Shots/Source/ShotsEditor/Private/UI/CineAssembly/SCineAssemblySchemaWindow.cpp
@@ -9,6 +9,7 @@
 #include "IStructureDetailsView.h"
 #include "Modules/ModuleManager.h"
 #include "NamingTokensEngineSubsystem.h"
+#include "PropertyEditorDelegates.h"
 #include "PropertyEditorModule.h"
 #include "ShotsEditorStyle.h"
 #include "Styling/StyleColors.h"
@@ -110,11 +111,11 @@ TSharedRef<SWidget> SCineAssemblySchemaWindow::MakeMenuPanel()
 				.MaxSegmentsPerLine(1)
 				.Value_Lambda([this]()
 				{ 
-					return TabSwitcher->GetActiveWidgetIndex(); 
+					return MenuTabSwitcher->GetActiveWidgetIndex();
 				})
 				.OnValueChanged_Lambda([this](int32 NewValue)
 				{
-					TabSwitcher->SetActiveWidgetIndex(NewValue);
+					MenuTabSwitcher->SetActiveWidgetIndex(NewValue);
 				})
 
 				+ SSegmentedControl<int32>::Slot(0)
@@ -131,7 +132,7 @@ TSharedRef<SWidget> SCineAssemblySchemaWindow::MakeMenuPanel()
 
 TSharedRef<SWidget> SCineAssemblySchemaWindow::MakeContentPanel()
 {
-	TabSwitcher = SNew(SWidgetSwitcher)
+	MenuTabSwitcher = SNew(SWidgetSwitcher)
 		+ SWidgetSwitcher::Slot()
 		[
 			MakeAssetTabContent()
@@ -147,7 +148,7 @@ TSharedRef<SWidget> SCineAssemblySchemaWindow::MakeContentPanel()
 		+ SVerticalBox::Slot()
 		.FillHeight(1.0f)
 		[
-			TabSwitcher.ToSharedRef()
+			MenuTabSwitcher.ToSharedRef()
 		]
 
 		+SVerticalBox::Slot()
@@ -233,6 +234,114 @@ TSharedRef<SWidget> SCineAssemblySchemaWindow::MakeButtonsPanel()
 }
 
 TSharedRef<SWidget> SCineAssemblySchemaWindow::MakeAssetTabContent()
+{
+	AssetTabSwitcher = SNew(SWidgetSwitcher)
+		+ SWidgetSwitcher::Slot()
+		[ MakeOverviewWidget() ]
+
+		+ SWidgetSwitcher::Slot()
+		[ MakeDetailsWidget(false) ]
+
+		+ SWidgetSwitcher::Slot()
+		[ MakeDetailsWidget(true) ];
+
+	return SNew(SVerticalBox)
+
+		+ SVerticalBox::Slot()
+		.AutoHeight()
+		.HAlign(HAlign_Center)
+		.Padding(8.0f)
+		[
+			SNew(SSegmentedControl<int32>)
+				.Value(0)
+				.OnValueChanged_Lambda([Switcher = AssetTabSwitcher](int32 NewValue)
+					{
+						Switcher->SetActiveWidgetIndex(NewValue);
+					})
+
+				+ SSegmentedControl<int32>::Slot(0)
+				.Text(LOCTEXT("OverviewTab", "Overview"))
+				.Icon(FSlateIcon(FShotsEditorStyle::StyleName, "Icons.Animation").GetIcon())
+
+				+ SSegmentedControl<int32>::Slot(1)
+				.Text(LOCTEXT("DetailsTab", "Details"))
+				.Icon(FSlateIcon(FAppStyle::GetAppStyleSetName(), "Icons.Details").GetIcon())
+
+				+ SSegmentedControl<int32>::Slot(2)
+				.Text(LOCTEXT("DataTab", "Data"))
+				.Icon(FSlateIcon(FShotsEditorStyle::StyleName, "Icons.DataAsset").GetIcon())
+		]
+
+		+ SVerticalBox::Slot()
+		.AutoHeight()
+		.HAlign(HAlign_Fill)
+		[
+			AssetTabSwitcher.ToSharedRef()
+		];
+}
+
+TSharedRef<SWidget> SCineAssemblySchemaWindow::MakeOverviewWidget()
+{
+	return SNew(SBorder)
+		.BorderImage(FShotsEditorStyle::Get().GetBrush("Borders.PanelNoBorder"))
+		.Padding(16.0f)
+		[
+			SNew(SVerticalBox)
+
+				// TODO: Populate this with an actual preview image sourced from the schema
+				+ SVerticalBox::Slot()
+				.Padding(0.0f, 0.0f, 0.0f, 16.0f)
+				.AutoHeight()
+				.HAlign(HAlign_Center)
+				[
+					SNew(SBox)
+					.WidthOverride(420.0f)
+					.HeightOverride(300.0f)
+						[
+							SNew(SBorder)
+								.BorderImage(FShotsEditorStyle::Get().GetBrush("Borders.Background"))
+								.HAlign(HAlign_Center)
+								.VAlign(VAlign_Center)
+								[
+									SNew(STextBlock).Text(LOCTEXT("PreviewImageText", "PREVIEW IMAGE"))
+								]
+						]
+				]
+	
+				+ SVerticalBox::Slot()
+				.AutoHeight()
+				.Padding(0.0f, 0.0f, 0.0f, 4.0f)
+				[
+					SNew(STextBlock)
+						.AutoWrapText(true)
+						.Text_Lambda([this]() -> FText
+							{
+								return FText::FromString(SchemaToConfigure->SchemaName);
+							})
+				]
+
+				+ SVerticalBox::Slot()
+				.AutoHeight()
+				.Padding(0.0f, 0.0f, 0.0f, 8.0f)
+				[
+					SNew(STextBlock)
+						.Text(LOCTEXT("SchemClassName", "Cine Assembly Schema"))
+						.ColorAndOpacity(FSlateColor::UseSubduedForeground())
+				]
+
+				+ SVerticalBox::Slot()
+				.AutoHeight()
+				[
+					SNew(STextBlock)
+						.Text_Lambda([this]() -> FText
+							{
+								return !SchemaToConfigure->Description.IsEmpty() ? FText::FromString(SchemaToConfigure->Description) : LOCTEXT("EmptyDescription", "No description");
+							})
+				]
+		];
+}
+
+TSharedRef<SWidget> SCineAssemblySchemaWindow::MakeDetailsWidget(bool bShowMetadata)
 {
 	FPropertyEditorModule& PropertyEditor = FModuleManager::Get().LoadModuleChecked<FPropertyEditorModule>(TEXT("PropertyEditor"));
 
@@ -241,13 +350,28 @@ TSharedRef<SWidget> SCineAssemblySchemaWindow::MakeAssetTabContent()
 	DetailsViewArgs.bHideSelectionTip = true;
 
 	TSharedRef<IDetailsView> DetailsView = PropertyEditor.CreateDetailView(DetailsViewArgs);
-	DetailsView->SetObject(SchemaToConfigure, true);
 
+	DetailsView->SetIsPropertyVisibleDelegate(FIsPropertyVisible::CreateSP(this, &SCineAssemblySchemaWindow::IsPropertyVisible, bShowMetadata));
+
+	DetailsView->SetObject(SchemaToConfigure, true);
 	DetailsView->OnFinishedChangingProperties().AddSP(this, &SCineAssemblySchemaWindow::OnSchemaPropertiesChanged);
 
 	return DetailsView;
 }
 
+bool SCineAssemblySchemaWindow::IsPropertyVisible(const FPropertyAndParent& PropertyAndParent, bool bShowMetadata)
+{
+	if (PropertyAndParent.Property.GetFName() == GET_MEMBER_NAME_CHECKED(UCineAssemblySchema, AssemblyMetadata))
+	{
+		return bShowMetadata;
+	}
+	else if ((PropertyAndParent.ParentProperties.Num() > 0) && (PropertyAndParent.ParentProperties[0]->GetFName() == GET_MEMBER_NAME_CHECKED(UCineAssemblySchema, AssemblyMetadata)))
+	{
+		return bShowMetadata;
+	}
+	return !bShowMetadata;
+}
+
 TSharedRef<SWidget> SCineAssemblySchemaWindow::MakeHierarchyTabContent()
 {
 	TreeView = SNew(STreeView<TSharedPtr<FSchemaTreeItem>>)
diff --git a/Engine/Plugins/VirtualProduction/Shots/Source/ShotsEditor/Private/UI/CineAssembly/SCineAssemblySchemaWindow.h b/Engine/Plugins/VirtualProduction/Shots/Source/ShotsEditor/Private/UI/CineAssembly/SCineAssemblySchemaWindow.h
--- a/Engine/Plugins/VirtualProduction/Shots/Source/ShotsEditor/Private/UI/CineAssembly/SCineAssemblySchemaWindow.h
+++ b/Engine/Plugins/VirtualProduction/Shots/Source/ShotsEditor/Private/UI/CineAssembly/SCineAssemblySchemaWindow.h
@@ -11,6 +11,8 @@
 #include "Widgets/Views/SListView.h"
 #include "Widgets/Views/STreeView.h"
 
+struct FPropertyAndParent;
+
 /** Mode that configures the UI based on the intended user interactions with the Cine Assembly Schema asset */
 enum class ESchemaConfigMode : uint8
 {
@@ -116,6 +118,15 @@ private:
 	/** Creates the content for the Hierarchy tab */
 	TSharedRef<SWidget> MakeHierarchyTabContent();
 
+	/** Creates the content for the Overview tab */
+	TSharedRef<SWidget> MakeOverviewWidget();
+
+	/** Creates the content for the Details and Data tabs (metadata properties are only shown in the Data tab) */
+	TSharedRef<SWidget> MakeDetailsWidget(bool bShowMetadata);
+
+	/** Filter used by the Details Views to determine which schema properties to display */
+	bool IsPropertyVisible(const FPropertyAndParent& PropertyAndParent, bool bShowMetadata);
+
 	/** Closes the window and indicates that a new asset should be created by the asset factory */
 	FReply OnCreateAssetClicked();
 
@@ -207,8 +218,11 @@ private:
 	FReply OnTreeRowAcceptDrop(const FDragDropEvent& InDragDropEvent, EItemDropZone InDropZone, TSharedPtr<FSchemaTreeItem> InItem);
 
 private:
-	/** Switcher that controls which tab widget is currently visible */
-	TSharedPtr<SWidgetSwitcher> TabSwitcher;
+	/** Switcher that controls which menu tab is currently visible */
+	TSharedPtr<SWidgetSwitcher> MenuTabSwitcher;
+
+	/** Switcher that controls which tab widget in the Asset Tab is currently visible */
+	TSharedPtr<SWidgetSwitcher> AssetTabSwitcher;
 
 	/** Whether a new asset should be created after closing this window */
 	bool bShouldCreateNewAsset = false;
