ShotManagement: Add static BP function to create a new Cine Assembly, specifying its schema, level, and metadata. This will support procedural assembly creation workflows.
#rnx

diff --git a/Engine/Plugins/VirtualProduction/Shots/Source/ShotsEditor/Private/CineAssembly.cpp b/Engine/Plugins/VirtualProduction/Shots/Source/ShotsEditor/Private/CineAssembly.cpp
--- a/Engine/Plugins/VirtualProduction/Shots/Source/ShotsEditor/Private/CineAssembly.cpp
+++ b/Engine/Plugins/VirtualProduction/Shots/Source/ShotsEditor/Private/CineAssembly.cpp
@@ -35,6 +35,7 @@ void UCineAssembly::SetSchema(UCineAssemblySchema* InSchema)
 	if (BaseSchema == nullptr)
 	{
 		BaseSchema = InSchema;
+		ChangeSchema(InSchema);
 	}
 }
 
diff --git a/Engine/Plugins/VirtualProduction/Shots/Source/ShotsEditor/Private/ProductionFunctionLibrary.cpp b/Engine/Plugins/VirtualProduction/Shots/Source/ShotsEditor/Private/ProductionFunctionLibrary.cpp
--- a/Engine/Plugins/VirtualProduction/Shots/Source/ShotsEditor/Private/ProductionFunctionLibrary.cpp
+++ b/Engine/Plugins/VirtualProduction/Shots/Source/ShotsEditor/Private/ProductionFunctionLibrary.cpp
@@ -2,6 +2,15 @@
 
 #include "ProductionFunctionLibrary.h"
 
+#include "Algo/Find.h"
+#include "AssetToolsModule.h"
+#include "CineAssembly.h"
+#include "CineAssemblyNamingTokens.h"
+#include "CineAssemblySchema.h"
+#include "MovieScene.h"
+#include "MovieSceneToolsProjectSettings.h"
+#include "NamingTokensEngineSubsystem.h"
+
 const TArray<FCinematicProduction> UProductionFunctionLibrary::GetAllProductions()
 {
 	const UProductionSettings* const ProductionSettings = GetDefault<UProductionSettings>();
@@ -77,3 +86,140 @@ void UProductionFunctionLibrary::RenameProduction(FGuid ProductionID, FString Ne
 	UProductionSettings* ProductionSettings = GetMutableDefault<UProductionSettings>();
 	ProductionSettings->RenameProduction(ProductionID, NewName);
 }
+
+UCineAssembly* UProductionFunctionLibrary::CreateAssembly(UCineAssemblySchema* Schema, TSoftObjectPtr<UWorld> Level, TMap<FString, FString> Metadata, const FString& Path, const FString& Name, bool bUseDefaultNameFromSchema)
+{
+	IAssetTools& AssetTools = FModuleManager::GetModuleChecked<FAssetToolsModule>("AssetTools").Get();
+
+	// We need a valid name for the initial asset creation, even if we are going to rename it later.
+	FString InitialName = Name;
+	if (InitialName.IsEmpty())
+	{
+		InitialName = TEXT("NewCineAssembly");
+	}
+
+	FString UniquePackageName;
+	FString UniqueAssetName;
+	AssetTools.CreateUniqueAssetName(Path / InitialName, TEXT(""), UniquePackageName, UniqueAssetName);
+
+	// Create the new assembly with a unique name. 
+	// No factory is provided, because we need to set the schema and metadata on the Assembly asset first, before we can create any SubAssemblies.
+	UCineAssembly* NewAssembly = Cast<UCineAssembly>(AssetTools.CreateAsset(UniqueAssetName, FPaths::GetPath(UniquePackageName), UCineAssembly::StaticClass(), nullptr));
+
+	NewAssembly->Level = Level.ToSoftObjectPath();
+	NewAssembly->SetSchema(Schema);
+
+	// Do the same setup that ULevelSequence assets do when they are created
+	NewAssembly->Initialize();
+
+	const UMovieSceneToolsProjectSettings* MovieSceneToolsProjectSettings = GetDefault<UMovieSceneToolsProjectSettings>();
+
+	FFrameRate TickResolution = NewAssembly->GetMovieScene()->GetTickResolution();
+	NewAssembly->GetMovieScene()->SetPlaybackRange((MovieSceneToolsProjectSettings->DefaultStartTime * TickResolution).FloorToFrame(), (MovieSceneToolsProjectSettings->DefaultDuration * TickResolution).FloorToFrame().Value);
+
+	AddMetadataToAssembly(NewAssembly, Metadata);
+
+	if (Schema && bUseDefaultNameFromSchema)
+	{
+		NewAssembly->AssemblyName.Template = Schema->DefaultAssemblyName;
+
+		if (!NewAssembly->AssemblyName.Template.IsEmpty())
+		{
+			RenameAssembly(NewAssembly);
+		}
+	}
+	else
+	{
+		NewAssembly->AssemblyName.Template = UniqueAssetName;
+		NewAssembly->AssemblyName.Resolved = FText::FromString(UniqueAssetName);
+	}
+
+	NewAssembly->CreateSubAssemblies();
+
+	return NewAssembly;
+}
+
+void UProductionFunctionLibrary::AddMetadataToAssembly(UCineAssembly* Assembly, TMap<FString, FString> Metadata)
+{
+	const UCineAssemblySchema* Schema = Assembly->GetSchema();
+
+	for (const TPair<FString, FString>& Pair : Metadata)
+	{
+		const FAssemblyMetadataDesc* MetadataDesc = Schema ? Algo::FindBy(Schema->AssemblyMetadata, Pair.Key, &FAssemblyMetadataDesc::Key) : nullptr;
+
+		if (MetadataDesc)
+		{
+			if (MetadataDesc->Type == ECineAssemblyMetadataType::String)
+			{
+				Assembly->AddMetadataAsString(Pair.Key, Pair.Value);
+			}
+			else if (MetadataDesc->Type == ECineAssemblyMetadataType::Bool)
+			{
+				Assembly->AddMetadataAsBool(Pair.Key, Pair.Value.ToBool());
+			}
+			else if (MetadataDesc->Type == ECineAssemblyMetadataType::Number)
+			{
+				Assembly->AddMetadataAsNumber(Pair.Key, FCString::Atof(*Pair.Value));
+			}
+			else if (MetadataDesc->Type == ECineAssemblyMetadataType::AssetPath)
+			{
+				Assembly->AddMetadataAsString(Pair.Key, Pair.Value);
+			}
+			else if (MetadataDesc->Type == ECineAssemblyMetadataType::CineAssembly)
+			{
+				Assembly->AddMetadataAsString(Pair.Key, Pair.Value);
+			}
+		}
+		else
+		{
+			// If the metadata string does not match anything specified by the schema, simply add the metadata as a string.
+			Assembly->InstanceMetadata.Add(*Pair.Key, Pair.Value);
+			Assembly->AddMetadataAsString(Pair.Key, Pair.Value);
+		}
+	}
+}
+
+void UProductionFunctionLibrary::RenameAssembly(UCineAssembly* Assembly)
+{
+	IAssetTools& AssetTools = FModuleManager::GetModuleChecked<FAssetToolsModule>("AssetTools").Get();
+
+	// Add the cine assembly being created to the naming token context
+	UCineAssemblyNamingTokensContext* NamingTokenContext = NewObject<UCineAssemblyNamingTokensContext>();
+	NamingTokenContext->Assembly = Assembly;
+
+	FNamingTokenFilterArgs FilterArgs;
+	FilterArgs.AdditionalNamespacesToInclude.Add(UCineAssemblyNamingTokens::TokenNamespace);
+
+	// Evaluate the default assembly name to resolve any tokens 
+	UNamingTokensEngineSubsystem* NamingTokensSubsystem = GEngine->GetEngineSubsystem<UNamingTokensEngineSubsystem>();
+	FNamingTokenResultData Result = NamingTokensSubsystem->EvaluateTokenString(Assembly->AssemblyName.Template, FilterArgs, { NamingTokenContext });
+
+	Assembly->AssemblyName.Resolved = Result.EvaluatedText;
+
+	// We now need to rename the asset (which was originally created with a default name) to the resolved default assembly name
+	const FSoftObjectPath OldObjectPath = FSoftObjectPath(Assembly);
+
+	// Ensure that the resolved assembly name is actually unique
+	const FString DesiredPackageName = FPaths::GetPath(Assembly->GetPathName()) / Assembly->AssemblyName.Resolved.ToString();
+
+	FString UniquePackageName;
+	FString UniqueAssetName;
+	AssetTools.CreateUniqueAssetName(DesiredPackageName, TEXT(""), UniquePackageName, UniqueAssetName);
+
+	// If we had to add anything to the assembly name to make it unique, update the assembly name template to match
+	if (UniqueAssetName != Assembly->AssemblyName.Resolved.ToString())
+	{
+		Assembly->AssemblyName.Template = UniqueAssetName;
+		Assembly->AssemblyName.Resolved = FText::FromString(UniqueAssetName);
+	}
+
+	const FTopLevelAssetPath NewAssetPath(*UniquePackageName, *UniqueAssetName);
+	const FSoftObjectPath NewObjectPath(NewAssetPath);
+
+	if (OldObjectPath != NewObjectPath)
+	{
+		FAssetRenameData RenameData = FAssetRenameData(OldObjectPath, NewObjectPath);
+		AssetTools.RenameAssets({ RenameData });
+	}
+}
+
diff --git a/Engine/Plugins/VirtualProduction/Shots/Source/ShotsEditor/Private/ProductionFunctionLibrary.h b/Engine/Plugins/VirtualProduction/Shots/Source/ShotsEditor/Private/ProductionFunctionLibrary.h
--- a/Engine/Plugins/VirtualProduction/Shots/Source/ShotsEditor/Private/ProductionFunctionLibrary.h
+++ b/Engine/Plugins/VirtualProduction/Shots/Source/ShotsEditor/Private/ProductionFunctionLibrary.h
@@ -7,6 +7,9 @@
 
 #include "ProductionFunctionLibrary.generated.h"
 
+class UCineAssembly;
+class UCineAssemblySchema;
+
 /**
  * Library of Blueprint/Python accessible functions to interface with the Cinematic Production Settings
  */
@@ -59,4 +62,17 @@ public:
 	/** Renames the Cinematic Production matching the input ProductionID */
 	UFUNCTION(BlueprintCallable, Category = "Shot Management")
 	static void RenameProduction(FGuid ProductionID, FString NewName);
+
+	/** 
+	 * Create a new CineAssembly asset using the input Schema, Level, and Metadata.
+	 * If bUseDefaultNameFromSchema is true, the default assembly name from the specified schema will be used as the new asset name.
+	 * It is important that any metadata required for resolving asset naming tokens is provided to this function so that the Assembly and SubAssemblies are all named correctly. 
+	 */
+	UFUNCTION(BlueprintCallable, Category = "Shot Management")
+	static UCineAssembly* CreateAssembly(UCineAssemblySchema* Schema, TSoftObjectPtr<UWorld> Level, TMap<FString, FString> Metadata, const FString& Path, const FString& Name, bool bUseDefaultNameFromSchema = true);
+
+private:
+
+	static void AddMetadataToAssembly(UCineAssembly* Assembly, TMap<FString, FString> Metadata);
+	static void RenameAssembly(UCineAssembly* Assembly);
 };
